CIS fixes for 2045858

diff --git a/Engine/Source/Developer/DesktopPlatform/Private/DesktopPlatformBase.h b/Engine/Source/Developer/DesktopPlatform/Private/DesktopPlatformBase.h
--- a/Engine/Source/Developer/DesktopPlatform/Private/DesktopPlatformBase.h
+++ b/Engine/Source/Developer/DesktopPlatform/Private/DesktopPlatformBase.h
@@ -30,5 +30,5 @@ private:
 	int32 ParseReleaseVersion(const FString &Version);
 
 	TSharedPtr<FJsonObject> LoadProjectFile(const FString &FileName);
-	bool FDesktopPlatformBase::SaveProjectFile(const FString &FileName, TSharedPtr<FJsonObject> Object);
+	bool SaveProjectFile(const FString &FileName, TSharedPtr<FJsonObject> Object);
 };
diff --git a/Engine/Source/Developer/DesktopPlatform/Private/Windows/DesktopPlatformWindows.cpp b/Engine/Source/Developer/DesktopPlatform/Private/Windows/DesktopPlatformWindows.cpp
--- a/Engine/Source/Developer/DesktopPlatform/Private/Windows/DesktopPlatformWindows.cpp
+++ b/Engine/Source/Developer/DesktopPlatform/Private/Windows/DesktopPlatformWindows.cpp
@@ -8,6 +8,8 @@
 	#include <Winver.h>
 #include "HideWindowsPlatformTypes.h"
 
+#pragma comment( lib, "version.lib" )
+
 #define LOCTEXT_NAMESPACE "DesktopPlatform"
 #define MAX_FILETYPES_STR 4096
 #define MAX_FILENAME_STR 65536 // This buffer has to be big enough to contain the names of all the selected files as well as the null characters between them and the null character at the end
