Fix win32 DesktopPlatform for Launcher builds.

diff --git a/Engine/Source/Developer/DesktopPlatform/Private/Windows/WindowsRegistry.cpp b/Engine/Source/Developer/DesktopPlatform/Private/Windows/WindowsRegistry.cpp
--- a/Engine/Source/Developer/DesktopPlatform/Private/Windows/WindowsRegistry.cpp
+++ b/Engine/Source/Developer/DesktopPlatform/Private/Windows/WindowsRegistry.cpp
@@ -113,7 +113,7 @@ bool FRegistryKey::Read(HKEY hKey)
 	EnumerateRegistryKeys(hKey, KeyNames);
 
 	// Read all the child keys
-	for (size_t Idx = 0; Idx < KeyNames.Num(); Idx++)
+	for (int32 Idx = 0; Idx < KeyNames.Num(); Idx++)
 	{
 		FRegistryKey *Key = FindOrAddKey(KeyNames[Idx]);
 		if (!Key->Read(hKey, KeyNames[Idx]))
@@ -127,7 +127,7 @@ bool FRegistryKey::Read(HKEY hKey)
 	EnumerateRegistryValues(hKey, ValueNames);
 
 	// Read all the values
-	for (size_t Idx = 0; Idx < ValueNames.Num(); Idx++)
+	for (int32 Idx = 0; Idx < ValueNames.Num(); Idx++)
 	{
 		FRegistryValue *Value = FindOrAddValue(ValueNames[Idx]);
 		if (!Value->Read(hKey, ValueNames[Idx]))
@@ -163,7 +163,7 @@ bool FRegistryKey::Write(HKEY hKey, bool bRemoveDifferences) const
 		EnumerateRegistryValues(hKey, ValueNames);
 
 		// Remove any values that don't exist any more
-		for (size_t Idx = 0; Idx < ValueNames.Num(); Idx++)
+		for (int32 Idx = 0; Idx < ValueNames.Num(); Idx++)
 		{
 			if (!Values.Contains(ValueNames[Idx]) && RegDeleteValue(hKey, *ValueNames[Idx]) != ERROR_SUCCESS)
 			{
@@ -176,7 +176,7 @@ bool FRegistryKey::Write(HKEY hKey, bool bRemoveDifferences) const
 		EnumerateRegistryKeys(hKey, KeyNames);
 
 		// Remove any keys that don't exist any more
-		for (size_t Idx = 0; Idx < KeyNames.Num(); Idx++)
+		for (int32 Idx = 0; Idx < KeyNames.Num(); Idx++)
 		{
 			if (!Keys.Contains(KeyNames[Idx]) && RegDeleteTree(hKey, *KeyNames[Idx]) != ERROR_SUCCESS)
 			{
@@ -229,7 +229,7 @@ bool FRegistryKey::IsUpToDate(HKEY hKey, bool bRemoveDifferences) const
 		EnumerateRegistryValues(hKey, ValueNames);
 
 		// Check there aren't any extra values
-		for (size_t Idx = 0; Idx < ValueNames.Num(); Idx++)
+		for (int32 Idx = 0; Idx < ValueNames.Num(); Idx++)
 		{
 			if (!Values.Contains(ValueNames[Idx]))
 			{
@@ -242,7 +242,7 @@ bool FRegistryKey::IsUpToDate(HKEY hKey, bool bRemoveDifferences) const
 		EnumerateRegistryKeys(hKey, KeyNames);
 
 		// Check there aren't any extra keys
-		for (size_t Idx = 0; Idx < KeyNames.Num(); Idx++)
+		for (int32 Idx = 0; Idx < KeyNames.Num(); Idx++)
 		{
 			if (!Keys.Contains(KeyNames[Idx]))
 			{
