Change default UnrealVersionSelector executable name to not have the Win64-Shipping suffix.

diff --git a/Engine/Source/Developer/DesktopPlatform/Private/Windows/DesktopPlatformWindows.cpp b/Engine/Source/Developer/DesktopPlatform/Private/Windows/DesktopPlatformWindows.cpp
--- a/Engine/Source/Developer/DesktopPlatform/Private/Windows/DesktopPlatformWindows.cpp
+++ b/Engine/Source/Developer/DesktopPlatform/Private/Windows/DesktopPlatformWindows.cpp
@@ -447,7 +447,7 @@ void FDesktopPlatformWindows::GetRequiredRegistrySettings(TIndirectArray<FRegist
 	FString DefaultVersionSelectorName = FPlatformProcess::ExecutableName(false);
 	if (!DefaultVersionSelectorName.StartsWith(TEXT("UnrealVersionSelector")))
 	{
-		DefaultVersionSelectorName = TEXT("UnrealVersionSelector-Win64-Shipping.exe");
+		DefaultVersionSelectorName = TEXT("UnrealVersionSelector.exe");
 	}
 	FString ExecutableFileName = FString(FPlatformProcess::BaseDir()) / DefaultVersionSelectorName;
 
@@ -459,7 +459,7 @@ void FDesktopPlatformWindows::GetRequiredRegistrySettings(TIndirectArray<FRegist
 		FString NormalizedInstallDir = InstallDir;
 		FPaths::NormalizeDirectoryName(NormalizedInstallDir);
 
-		FString InstalledExecutableFileName = NormalizedInstallDir / FString("Launcher/Engine/Binaries/Win64/UnrealVersionSelector-Win64-Shipping.exe");
+		FString InstalledExecutableFileName = NormalizedInstallDir / FString("Launcher/Engine/Binaries/Win64/UnrealVersionSelector.exe");
 		if(GetShellIntegrationVersion(InstalledExecutableFileName) == GetShellIntegrationVersion(ExecutableFileName))
 		{
 			ExecutableFileName = InstalledExecutableFileName;
