CIS fix for mac.

diff --git a/Engine/Source/Developer/DesktopPlatform/Private/Mac/DesktopPlatformMac.cpp b/Engine/Source/Developer/DesktopPlatform/Private/Mac/DesktopPlatformMac.cpp
--- a/Engine/Source/Developer/DesktopPlatform/Private/Mac/DesktopPlatformMac.cpp
+++ b/Engine/Source/Developer/DesktopPlatform/Private/Mac/DesktopPlatformMac.cpp
@@ -517,12 +517,12 @@ void FDesktopPlatformMac::EnumerateEngineInstallations(TMap<FString, FString> &O
 	EnumerateLauncherEngineInstallations(OutInstallations);
 }
 
-bool FDesktopPlatformMac::VerifyFileAssociations() OVERRIDE
+bool FDesktopPlatformMac::VerifyFileAssociations()
 {
 	return true;
 }
 
-bool FDesktopPlatformMac::UpdateFileAssociations() OVERRIDE
+bool FDesktopPlatformMac::UpdateFileAssociations()
 {
 	return true;
 }
