fixed UnrealEngineLauncher overwriting local P4 workspace files with files downloaded from the internet if launched through an Editor that was started with -Dev

diff --git a/Engine/Source/Developer/DesktopPlatform/Private/Mac/DesktopPlatformMac.cpp b/Engine/Source/Developer/DesktopPlatform/Private/Mac/DesktopPlatformMac.cpp
--- a/Engine/Source/Developer/DesktopPlatform/Private/Mac/DesktopPlatformMac.cpp
+++ b/Engine/Source/Developer/DesktopPlatform/Private/Mac/DesktopPlatformMac.cpp
@@ -348,7 +348,7 @@ void OpenLauncherCommandLine( const FString& InCommandLine )
 	
 }
 
-bool FDesktopPlatformMac::OpenLauncher(bool Install, const FString& CommandLineParams )
+bool FDesktopPlatformMac::OpenLauncher(bool Install, FString CommandLineParams )
 {
 	// If the launcher is already running, bring it to front
 	NSArray* RunningLaunchers = [NSRunningApplication runningApplicationsWithBundleIdentifier: @"com.epicgames.UnrealEngineLauncher"];
@@ -366,6 +366,7 @@ bool FDesktopPlatformMac::OpenLauncher(bool Install, const FString& CommandLineP
 	if (FParse::Param(FCommandLine::Get(), TEXT("Dev")))
 	{
 		LaunchPath = FPaths::Combine(*FPaths::EngineDir(), TEXT("Binaries"), TEXT("Mac"), TEXT("UnrealEngineLauncher-Mac-Debug.app/Contents/MacOS/UnrealEngineLauncher-Mac-Debug"));
+		CommandLineParams += TEXT(" -noselfupdate");
 	}
 	else
 	{
diff --git a/Engine/Source/Developer/DesktopPlatform/Private/Mac/DesktopPlatformMac.h b/Engine/Source/Developer/DesktopPlatform/Private/Mac/DesktopPlatformMac.h
--- a/Engine/Source/Developer/DesktopPlatform/Private/Mac/DesktopPlatformMac.h
+++ b/Engine/Source/Developer/DesktopPlatform/Private/Mac/DesktopPlatformMac.h
@@ -12,7 +12,7 @@ public:
 	virtual bool SaveFileDialog(const void* ParentWindowHandle, const FString& DialogTitle, const FString& DefaultPath, const FString& DefaultFile, const FString& FileTypes, uint32 Flags, TArray<FString>& OutFilenames) OVERRIDE;
 	virtual bool OpenDirectoryDialog(const void* ParentWindowHandle, const FString& DialogTitle, const FString& DefaultPath, FString& OutFolderName) OVERRIDE;
 	virtual bool OpenFontDialog(const void* ParentWindowHandle, FString& OutFontName, float& OutHeight, EFontImportFlags::Type& OutFlags) OVERRIDE;
-	virtual bool OpenLauncher(bool Install, const FString& CommandLineParams ) OVERRIDE;
+	virtual bool OpenLauncher(bool Install, FString CommandLineParams ) OVERRIDE;
 
 	virtual void EnumerateEngineInstallations(TMap<FString, FString> &OutInstallations) OVERRIDE;
 
diff --git a/Engine/Source/Developer/DesktopPlatform/Private/Windows/DesktopPlatformWindows.cpp b/Engine/Source/Developer/DesktopPlatform/Private/Windows/DesktopPlatformWindows.cpp
--- a/Engine/Source/Developer/DesktopPlatform/Private/Windows/DesktopPlatformWindows.cpp
+++ b/Engine/Source/Developer/DesktopPlatform/Private/Windows/DesktopPlatformWindows.cpp
@@ -129,13 +129,14 @@ bool FDesktopPlatformWindows::OpenFontDialog(const void* ParentWindowHandle, FSt
 	return bSuccess;
 }
 
-bool FDesktopPlatformWindows::OpenLauncher(bool Install, const FString& CommandLineParams )
+bool FDesktopPlatformWindows::OpenLauncher(bool Install, FString CommandLineParams )
 {
 	FString LaunchPath;
 
 	if (FParse::Param(FCommandLine::Get(), TEXT("Dev")))
 	{
 		LaunchPath = FPaths::Combine(*FPaths::EngineDir(), TEXT("Binaries"), TEXT("Win64"), TEXT("UnrealEngineLauncher-Win64-Debug.exe"));
+		CommandLineParams += TEXT(" -noselfupdate");
 	}
 	else
 	{
diff --git a/Engine/Source/Developer/DesktopPlatform/Private/Windows/DesktopPlatformWindows.h b/Engine/Source/Developer/DesktopPlatform/Private/Windows/DesktopPlatformWindows.h
--- a/Engine/Source/Developer/DesktopPlatform/Private/Windows/DesktopPlatformWindows.h
+++ b/Engine/Source/Developer/DesktopPlatform/Private/Windows/DesktopPlatformWindows.h
@@ -13,7 +13,7 @@ public:
 	virtual bool SaveFileDialog(const void* ParentWindowHandle, const FString& DialogTitle, const FString& DefaultPath, const FString& DefaultFile, const FString& FileTypes, uint32 Flags, TArray<FString>& OutFilenames) OVERRIDE;
 	virtual bool OpenDirectoryDialog(const void* ParentWindowHandle, const FString& DialogTitle, const FString& DefaultPath, FString& OutFolderName) OVERRIDE;
 	virtual bool OpenFontDialog(const void* ParentWindowHandle, FString& OutFontName, float& OutHeight, EFontImportFlags::Type& OutFlags) OVERRIDE;
-	virtual bool OpenLauncher(bool Install, const FString& CommandLineParams ) OVERRIDE;
+	virtual bool OpenLauncher(bool Install, FString CommandLineParams ) OVERRIDE;
 
 	virtual void EnumerateEngineInstallations(TMap<FString, FString> &OutInstallations) OVERRIDE;
 
diff --git a/Engine/Source/Developer/DesktopPlatform/Public/IDesktopPlatform.h b/Engine/Source/Developer/DesktopPlatform/Public/IDesktopPlatform.h
--- a/Engine/Source/Developer/DesktopPlatform/Public/IDesktopPlatform.h
+++ b/Engine/Source/Developer/DesktopPlatform/Public/IDesktopPlatform.h
@@ -115,7 +115,7 @@ public:
 	 * @param Install					Whether to install the marketplace if it is missing.
 	 * @return true if the marketplace was opened, false if it is not installed or could not be installed/opened.
 	 */
-	virtual bool OpenLauncher(bool Install, const FString& CommandLineParams ) = 0;
+	virtual bool OpenLauncher(bool Install, FString CommandLineParams ) = 0;
 
 	/**
 	* Gets the identifier for the currently executing engine installation
