Changed the path for temp project created by UE4EditorServices

diff --git a/Engine/Source/Developer/DesktopPlatform/Private/Mac/DesktopPlatformMac.cpp b/Engine/Source/Developer/DesktopPlatform/Private/Mac/DesktopPlatformMac.cpp
--- a/Engine/Source/Developer/DesktopPlatform/Private/Mac/DesktopPlatformMac.cpp
+++ b/Engine/Source/Developer/DesktopPlatform/Private/Mac/DesktopPlatformMac.cpp
@@ -518,7 +518,7 @@ void FDesktopPlatformMac::EnumerateEngineInstallations(TMap<FString, FString> &O
 	EnumerateLauncherEngineInstallations(OutInstallations);
 
 	// Create temp .uproject file to use with LSCopyApplicationURLsForURL
-	FString UProjectPath = FPaths::GameIntermediateDir() / "Unreal.uproject";
+	FString UProjectPath = FPlatformProcess::ApplicationSettingsDir() / "Unreal.uproject";
 	FArchive* File = IFileManager::Get().CreateFileWriter(*UProjectPath, FILEWRITE_EvenIfReadOnly);
 	if (File)
 	{
diff --git a/Engine/Source/Programs/Mac/UE4EditorServices/Resources/Mac/Info.plist b/Engine/Source/Programs/Mac/UE4EditorServices/Resources/Mac/Info.plist
--- a/Engine/Source/Programs/Mac/UE4EditorServices/Resources/Mac/Info.plist
+++ b/Engine/Source/Programs/Mac/UE4EditorServices/Resources/Mac/Info.plist
@@ -34,11 +34,11 @@
 	<key>CFBundlePackageType</key>
 	<string>APPL</string>
 	<key>CFBundleShortVersionString</key>
-	<string>1.1</string>
+	<string>1.2</string>
 	<key>CFBundleSignature</key>
 	<string>????</string>
 	<key>CFBundleVersion</key>
-	<string>1.1</string>
+	<string>1.2</string>
 	<key>LSApplicationCategoryType</key>
 	<string>public.app-category.developer-tools</string>
 	<key>LSBackgroundOnly</key>
