Updated UE4EditorServices binaries and compile fix

diff --git a/Engine/Source/Developer/DesktopPlatform/Private/Mac/DesktopPlatformMac.cpp b/Engine/Source/Developer/DesktopPlatform/Private/Mac/DesktopPlatformMac.cpp
--- a/Engine/Source/Developer/DesktopPlatform/Private/Mac/DesktopPlatformMac.cpp
+++ b/Engine/Source/Developer/DesktopPlatform/Private/Mac/DesktopPlatformMac.cpp
@@ -518,7 +518,7 @@ void FDesktopPlatformMac::EnumerateEngineInstallations(TMap<FString, FString> &O
 	EnumerateLauncherEngineInstallations(OutInstallations);
 
 	// Create temp .uproject file to use with LSCopyApplicationURLsForURL
-	FString UProjectPath = FPlatformProcess::ApplicationSettingsDir() / "Unreal.uproject";
+	FString UProjectPath = FString(FPlatformProcess::ApplicationSettingsDir()) / "Unreal.uproject";
 	FArchive* File = IFileManager::Get().CreateFileWriter(*UProjectPath, FILEWRITE_EvenIfReadOnly);
 	if (File)
 	{
