Add engine support for the launcher installation list rather than manually parsing manifests.

diff --git a/Engine/Source/Developer/DesktopPlatform/Private/DesktopPlatformBase.cpp b/Engine/Source/Developer/DesktopPlatform/Private/DesktopPlatformBase.cpp
--- a/Engine/Source/Developer/DesktopPlatform/Private/DesktopPlatformBase.cpp
+++ b/Engine/Source/Developer/DesktopPlatform/Private/DesktopPlatformBase.cpp
@@ -12,9 +12,25 @@ FString FDesktopPlatformBase::GetCurrentEngineIdentifier()
 
 void FDesktopPlatformBase::EnumerateLauncherEngineInstallations(TMap<FString, FString> &OutInstallations)
 {
-	// Just enumerate the known manifests for the time being 
-	CheckForLauncherEngineInstallation(TEXT("40003"), TEXT("4.0"), OutInstallations);
-	CheckForLauncherEngineInstallation(TEXT("1040003"), TEXT("4.1"), OutInstallations);
+	TMap<FString, FString> AllInstallations;
+	if(ReadLauncherInstallationList(AllInstallations))
+	{
+		// We've got a list of launcher installations. Filter it by the engine installations.
+		for(TMap<FString, FString>::TConstIterator Iter(AllInstallations); Iter; ++Iter)
+		{
+			FString AppName = Iter.Key();
+			if(AppName.RemoveFromStart(TEXT("UE_"), ESearchCase::CaseSensitive))
+			{
+				OutInstallations.Add(AppName, Iter.Value());
+			}
+		}
+	}
+	else
+	{
+		// Otherwise just enumerate the known manifests for the time being 
+		CheckForLauncherEngineInstallation(TEXT("40003"), TEXT("4.0"), OutInstallations);
+		CheckForLauncherEngineInstallation(TEXT("1040003"), TEXT("4.1"), OutInstallations);
+	}
 }
 
 bool FDesktopPlatformBase::GetEngineRootDirFromIdentifier(const FString &Identifier, FString &OutRootDir)
@@ -192,6 +208,42 @@ bool FDesktopPlatformBase::GetEngineIdentifierForProject(const FString &ProjectF
 	return false;
 }
 
+bool FDesktopPlatformBase::ReadLauncherInstallationList(TMap<FString, FString> &OutInstallations)
+{
+	FString InstalledListFile = FString(FPlatformProcess::ApplicationSettingsDir()) / TEXT("UnrealEngineLauncher/LauncherInstalled.dat");
+
+	// Read the installation manifest
+	FString InstalledText;
+	if (!FFileHelper::LoadFileToString(InstalledText, *InstalledListFile))
+	{
+		return false;
+	}
+
+	// Deserialize the object
+	TSharedPtr< FJsonObject > RootObject;
+	TSharedRef< TJsonReader<> > Reader = TJsonReaderFactory<>::Create(InstalledText);
+	if (!FJsonSerializer::Deserialize(Reader, RootObject) || !RootObject.IsValid())
+	{
+		return false;
+	}
+
+	// Parse the list of installations
+	TArray< TSharedPtr<FJsonValue> > InstallationList = RootObject->GetArrayField(TEXT("InstallationList"));
+	for(int32 Idx = 0; Idx < InstallationList.Num(); Idx++)
+	{
+		TSharedPtr<FJsonObject> InstallationItem = InstallationList[Idx]->AsObject();
+
+		FString AppName = InstallationItem->GetStringField(TEXT("AppName"));
+		FString InstallLocation = InstallationItem->GetStringField(TEXT("InstallLocation"));
+		if(AppName.Len() > 0 && InstallLocation.Len() > 0)
+		{
+			FPaths::NormalizeDirectoryName(InstallLocation);
+			OutInstallations.Add(AppName, InstallLocation);
+		}
+	}
+	return true;
+}
+
 void FDesktopPlatformBase::CheckForLauncherEngineInstallation(const FString &AppId, const FString &Identifier, TMap<FString, FString> &OutInstallations)
 {
 	FString ManifestText;
diff --git a/Engine/Source/Developer/DesktopPlatform/Private/DesktopPlatformBase.h b/Engine/Source/Developer/DesktopPlatform/Private/DesktopPlatformBase.h
--- a/Engine/Source/Developer/DesktopPlatform/Private/DesktopPlatformBase.h
+++ b/Engine/Source/Developer/DesktopPlatform/Private/DesktopPlatformBase.h
@@ -27,6 +27,7 @@ public:
 	virtual bool GetEngineIdentifierForProject(const FString &ProjectFileName, FString &OutIdentifier) OVERRIDE;
 
 private:
+	bool ReadLauncherInstallationList(TMap<FString, FString> &OutInstallations);
 	void CheckForLauncherEngineInstallation(const FString &AppId, const FString &Identifier, TMap<FString, FString> &OutInstallations);
 	int32 ParseReleaseVersion(const FString &Version);
 
