Fixed UnrealVersionSelector failing to add UE4 installation paths

FDesktopPlatformWindows::RegisterEngineInstallation was checking that creating the registry key had failed (rather than succeeded) before trying to add the value to the key.

diff --git a/Engine/Source/Developer/DesktopPlatform/Private/Windows/DesktopPlatformWindows.cpp b/Engine/Source/Developer/DesktopPlatform/Private/Windows/DesktopPlatformWindows.cpp
--- a/Engine/Source/Developer/DesktopPlatform/Private/Windows/DesktopPlatformWindows.cpp
+++ b/Engine/Source/Developer/DesktopPlatform/Private/Windows/DesktopPlatformWindows.cpp
@@ -357,7 +357,7 @@ bool FDesktopPlatformWindows::RegisterEngineInstallation(const FString &RootDir,
 	if(IsValidRootDirectory(RootDir))
 	{
 		HKEY hRootKey;
-		if(RegCreateKeyEx(HKEY_CURRENT_USER, InstallationsSubKey, 0, NULL, REG_OPTION_NON_VOLATILE, KEY_ALL_ACCESS, NULL, &hRootKey, NULL) != ERROR_SUCCESS)
+		if(RegCreateKeyEx(HKEY_CURRENT_USER, InstallationsSubKey, 0, NULL, REG_OPTION_NON_VOLATILE, KEY_ALL_ACCESS, NULL, &hRootKey, NULL) == ERROR_SUCCESS)
 		{
 			FString NewIdentifier = FGuid::NewGuid().ToString(EGuidFormats::DigitsWithHyphensInBraces);
 			LRESULT SetResult = RegSetValueEx(hRootKey, *NewIdentifier, 0, REG_SZ, (const BYTE*)*RootDir, (RootDir.Len() + 1) * sizeof(TCHAR));
