Implemented FDesktopPlatformMac::RegisterEngineInstallation()

diff --git a/Engine/Source/Developer/DesktopPlatform/Private/Mac/DesktopPlatformMac.cpp b/Engine/Source/Developer/DesktopPlatform/Private/Mac/DesktopPlatformMac.cpp
--- a/Engine/Source/Developer/DesktopPlatform/Private/Mac/DesktopPlatformMac.cpp
+++ b/Engine/Source/Developer/DesktopPlatform/Private/Mac/DesktopPlatformMac.cpp
@@ -515,7 +515,21 @@ bool FDesktopPlatformMac::FileDialogShared(bool bSave, const void* ParentWindowH
 
 bool FDesktopPlatformMac::RegisterEngineInstallation(const FString &RootDir, FString &OutIdentifier)
 {
-	return false;
+	bool bRes = false;
+	if (IsValidRootDirectory(RootDir))
+	{
+		FConfigFile ConfigFile;
+		FString ConfigPath = FString(FPlatformProcess::ApplicationSettingsDir()) / FString(TEXT("UnrealEngine")) / FString(TEXT("Install.ini"));
+		ConfigFile.Read(ConfigPath);
+
+		FConfigSection &Section = ConfigFile.FindOrAdd(TEXT("Installations"));
+		OutIdentifier = FGuid::NewGuid().ToString(EGuidFormats::DigitsWithHyphens);
+		Section.AddUnique(*OutIdentifier, RootDir);
+
+		ConfigFile.Dirty = true;
+		ConfigFile.Write(ConfigPath);
+	}
+	return bRes;
 }
 
 void FDesktopPlatformMac::EnumerateEngineInstallations(TMap<FString, FString> &OutInstallations)
