Don't re-read launcher installation list if the timestamp hasn't changed.

diff --git a/Engine/Source/Developer/DesktopPlatform/Private/DesktopPlatformBase.cpp b/Engine/Source/Developer/DesktopPlatform/Private/DesktopPlatformBase.cpp
--- a/Engine/Source/Developer/DesktopPlatform/Private/DesktopPlatformBase.cpp
+++ b/Engine/Source/Developer/DesktopPlatform/Private/DesktopPlatformBase.cpp
@@ -8,6 +8,11 @@
 
 #define LOCTEXT_NAMESPACE "DesktopPlatform"
 
+FDesktopPlatformBase::FDesktopPlatformBase()
+{
+	LauncherInstallationTimestamp = FDateTime::MinValue();
+}
+
 FString FDesktopPlatformBase::GetCurrentEngineIdentifier()
 {
 	if(CurrentEngineIdentifier.Len() == 0 && !GetEngineIdentifierFromRootDir(FPlatformMisc::RootDir(), CurrentEngineIdentifier))
@@ -19,40 +24,32 @@ FString FDesktopPlatformBase::GetCurrentEngineIdentifier()
 
 void FDesktopPlatformBase::EnumerateLauncherEngineInstallations(TMap<FString, FString> &OutInstallations)
 {
-	TMap<FString, FString> AllInstallations;
-	if(ReadLauncherInstallationList(AllInstallations))
+	// Cache the launcher install list if necessary
+	ReadLauncherInstallationList();
+
+	// We've got a list of launcher installations. Filter it by the engine installations.
+	for(TMap<FString, FString>::TConstIterator Iter(LauncherInstallationList); Iter; ++Iter)
 	{
-		// We've got a list of launcher installations. Filter it by the engine installations.
-		for(TMap<FString, FString>::TConstIterator Iter(AllInstallations); Iter; ++Iter)
+		FString AppName = Iter.Key();
+		if(AppName.RemoveFromStart(TEXT("UE_"), ESearchCase::CaseSensitive))
 		{
-			FString AppName = Iter.Key();
-			if(AppName.RemoveFromStart(TEXT("UE_"), ESearchCase::CaseSensitive))
-			{
-				OutInstallations.Add(AppName, Iter.Value());
-			}
+			OutInstallations.Add(AppName, Iter.Value());
 		}
 	}
-	else
-	{
-		// Otherwise just enumerate the known manifests for the time being 
-		CheckForLauncherEngineInstallation(TEXT("40003"), TEXT("4.0"), OutInstallations);
-		CheckForLauncherEngineInstallation(TEXT("1040003"), TEXT("4.1"), OutInstallations);
-	}
 }
 
 void FDesktopPlatformBase::EnumerateLauncherSampleInstallations(TArray<FString> &OutInstallations)
 {
-	TMap<FString, FString> AllInstallations;
-	if(ReadLauncherInstallationList(AllInstallations))
+	// Cache the launcher install list if necessary
+	ReadLauncherInstallationList();
+
+	// We've got a list of launcher installations. Filter it by the engine installations.
+	for(TMap<FString, FString>::TConstIterator Iter(LauncherInstallationList); Iter; ++Iter)
 	{
-		// We've got a list of launcher installations. Filter it by non-engine installations.
-		for(TMap<FString, FString>::TConstIterator Iter(AllInstallations); Iter; ++Iter)
+		FString AppName = Iter.Key();
+		if(!AppName.StartsWith(TEXT("UE_"), ESearchCase::CaseSensitive))
 		{
-			FString AppName = Iter.Key();
-			if(!AppName.StartsWith(TEXT("UE_"), ESearchCase::CaseSensitive))
-			{
-				OutInstallations.Add(Iter.Value());
-			}
+			OutInstallations.Add(Iter.Value());
 		}
 	}
 }
@@ -317,40 +314,49 @@ bool FDesktopPlatformBase::GenerateProjectFiles(const FString& RootDir, const FS
 	return bRes;
 }
 
-bool FDesktopPlatformBase::ReadLauncherInstallationList(TMap<FString, FString> &OutInstallations)
+void FDesktopPlatformBase::ReadLauncherInstallationList()
 {
 	FString InstalledListFile = FString(FPlatformProcess::ApplicationSettingsDir()) / TEXT("UnrealEngineLauncher/LauncherInstalled.dat");
 
-	// Read the installation manifest
-	FString InstalledText;
-	if (!FFileHelper::LoadFileToString(InstalledText, *InstalledListFile))
-	{
-		return false;
-	}
-
-	// Deserialize the object
-	TSharedPtr< FJsonObject > RootObject;
-	TSharedRef< TJsonReader<> > Reader = TJsonReaderFactory<>::Create(InstalledText);
-	if (!FJsonSerializer::Deserialize(Reader, RootObject) || !RootObject.IsValid())
+	// If the file does not exist, manually check for the 4.0 or 4.1 manifest
+	FDateTime NewListTimestamp = IFileManager::Get().GetTimeStamp(*InstalledListFile);
+	if(NewListTimestamp == FDateTime::MinValue())
 	{
-		return false;
+		if(LauncherInstallationList.Num() == 0)
+		{
+			CheckForLauncherEngineInstallation(TEXT("40003"), TEXT("UE_4.0"), LauncherInstallationList);
+			CheckForLauncherEngineInstallation(TEXT("1040003"), TEXT("UE_4.1"), LauncherInstallationList);
+		}
 	}
-
-	// Parse the list of installations
-	TArray< TSharedPtr<FJsonValue> > InstallationList = RootObject->GetArrayField(TEXT("InstallationList"));
-	for(int32 Idx = 0; Idx < InstallationList.Num(); Idx++)
+	else if(NewListTimestamp != LauncherInstallationTimestamp)
 	{
-		TSharedPtr<FJsonObject> InstallationItem = InstallationList[Idx]->AsObject();
-
-		FString AppName = InstallationItem->GetStringField(TEXT("AppName"));
-		FString InstallLocation = InstallationItem->GetStringField(TEXT("InstallLocation"));
-		if(AppName.Len() > 0 && InstallLocation.Len() > 0)
+		// Read the installation manifest
+		FString InstalledText;
+		if (FFileHelper::LoadFileToString(InstalledText, *InstalledListFile))
 		{
-			FPaths::NormalizeDirectoryName(InstallLocation);
-			OutInstallations.Add(AppName, InstallLocation);
+			// Deserialize the object
+			TSharedPtr< FJsonObject > RootObject;
+			TSharedRef< TJsonReader<> > Reader = TJsonReaderFactory<>::Create(InstalledText);
+			if (FJsonSerializer::Deserialize(Reader, RootObject) && RootObject.IsValid())
+			{
+				// Parse the list of installations
+				TArray< TSharedPtr<FJsonValue> > InstallationList = RootObject->GetArrayField(TEXT("InstallationList"));
+				for(int32 Idx = 0; Idx < InstallationList.Num(); Idx++)
+				{
+					TSharedPtr<FJsonObject> InstallationItem = InstallationList[Idx]->AsObject();
+
+					FString AppName = InstallationItem->GetStringField(TEXT("AppName"));
+					FString InstallLocation = InstallationItem->GetStringField(TEXT("InstallLocation"));
+					if(AppName.Len() > 0 && InstallLocation.Len() > 0)
+					{
+						FPaths::NormalizeDirectoryName(InstallLocation);
+						LauncherInstallationList.Add(AppName, InstallLocation);
+					}
+				}
+			}
+			LauncherInstallationTimestamp = NewListTimestamp;
 		}
 	}
-	return true;
 }
 
 void FDesktopPlatformBase::CheckForLauncherEngineInstallation(const FString &AppId, const FString &Identifier, TMap<FString, FString> &OutInstallations)
diff --git a/Engine/Source/Developer/DesktopPlatform/Private/DesktopPlatformBase.h b/Engine/Source/Developer/DesktopPlatform/Private/DesktopPlatformBase.h
--- a/Engine/Source/Developer/DesktopPlatform/Private/DesktopPlatformBase.h
+++ b/Engine/Source/Developer/DesktopPlatform/Private/DesktopPlatformBase.h
@@ -9,6 +9,8 @@
 class FDesktopPlatformBase : public IDesktopPlatform
 {
 public:
+	FDesktopPlatformBase();
+
 	// IDesktopPlatform Implementation
 	virtual FString GetCurrentEngineIdentifier() OVERRIDE;
 
@@ -34,9 +36,11 @@ public:
 
 private:
 	FString CurrentEngineIdentifier;
+	FDateTime LauncherInstallationTimestamp;
+	TMap<FString, FString> LauncherInstallationList;
 	TMap<FString, FUProjectDictionary> CachedProjectDictionaries;
 
-	bool ReadLauncherInstallationList(TMap<FString, FString> &OutInstallations);
+	void ReadLauncherInstallationList();
 	void CheckForLauncherEngineInstallation(const FString &AppId, const FString &Identifier, TMap<FString, FString> &OutInstallations);
 	int32 ParseReleaseVersion(const FString &Version);
 
diff --git a/Engine/Source/Developer/DesktopPlatform/Private/Windows/DesktopPlatformWindows.cpp b/Engine/Source/Developer/DesktopPlatform/Private/Windows/DesktopPlatformWindows.cpp
--- a/Engine/Source/Developer/DesktopPlatform/Private/Windows/DesktopPlatformWindows.cpp
+++ b/Engine/Source/Developer/DesktopPlatform/Private/Windows/DesktopPlatformWindows.cpp
@@ -480,7 +480,7 @@ bool FDesktopPlatformWindows::RunUnrealBuildTool(const FText& Description, const
 
 	// Spawn UBT
 	int32 ExitCode = 0;
-	return FFeedbackContextMarkup::PipeProcessOutput(Description, UnrealBuildToolPath, Arguments, Warn, &ExitCode) && ExitCode == 0 && false;
+	return FFeedbackContextMarkup::PipeProcessOutput(Description, UnrealBuildToolPath, Arguments, Warn, &ExitCode) && ExitCode == 0;
 }
 
 FFeedbackContext* FDesktopPlatformWindows::GetNativeFeedbackContext()
