CIS fix for Win32.

diff --git a/Engine/Source/Developer/DesktopPlatform/Private/Windows/WindowsNativeFeedbackContext.h b/Engine/Source/Developer/DesktopPlatform/Private/Windows/WindowsNativeFeedbackContext.h
--- a/Engine/Source/Developer/DesktopPlatform/Private/Windows/WindowsNativeFeedbackContext.h
+++ b/Engine/Source/Developer/DesktopPlatform/Private/Windows/WindowsNativeFeedbackContext.h
@@ -72,7 +72,7 @@ private:
 	void CreateSlowTaskWindow(const FText &InStatus, bool bInShowCancelButton);
 	void DestroySlowTaskWindow();
 
-	static DWORD SlowTaskThreadProc(void *Params);
+	static DWORD WINAPI SlowTaskThreadProc(void *Params);
 	static LRESULT CALLBACK SlowTaskWindowProc(HWND hWnd, UINT Msg, WPARAM wParam, LPARAM lParam);
 	static void LayoutControls(HWND hWnd, const FWindowParams* Params);
 };
