Add a Mac FFeedbackContext which shows the log window during slow tasks. Windows version displays a nice status message and progress bar, but this will do for now. Will be used to recompile out-of-date binaries on startup.

diff --git a/Engine/Source/Developer/DesktopPlatform/Private/Mac/DesktopPlatformMac.cpp b/Engine/Source/Developer/DesktopPlatform/Private/Mac/DesktopPlatformMac.cpp
--- a/Engine/Source/Developer/DesktopPlatform/Private/Mac/DesktopPlatformMac.cpp
+++ b/Engine/Source/Developer/DesktopPlatform/Private/Mac/DesktopPlatformMac.cpp
@@ -3,6 +3,7 @@
 #include "DesktopPlatformPrivatePCH.h"
 #include "MacApplication.h"
 #include "FeedbackContextMarkup.h"
+#include "MacNativeFeedbackContext.h"
 
 #define LOCTEXT_NAMESPACE "DesktopPlatform"
 
@@ -641,7 +642,8 @@ bool FDesktopPlatformMac::RunUnrealBuildTool(const FText& Description, const FSt
 
 FFeedbackContext* FDesktopPlatformMac::GetNativeFeedbackContext()
 {
-	return GWarn;
+	static FMacNativeFeedbackContext Warn;
+	return &Warn;
 }
 
 #undef LOCTEXT_NAMESPACE
diff --git a/Engine/Source/Developer/DesktopPlatform/Private/Mac/MacNativeFeedbackContext.cpp b/Engine/Source/Developer/DesktopPlatform/Private/Mac/MacNativeFeedbackContext.cpp
new file mode 100644
--- /dev/null
+++ b/Engine/Source/Developer/DesktopPlatform/Private/Mac/MacNativeFeedbackContext.cpp
@@ -0,0 +1,81 @@
+// Copyright 1998-2014 Epic Games, Inc. All Rights Reserved.
+
+#include "DesktopPlatformPrivatePCH.h"
+#include "MacNativeFeedbackContext.h"
+
+FMacNativeFeedbackContext::FMacNativeFeedbackContext()
+	: FFeedbackContext()
+	, Context( NULL )
+	, SlowTaskCount( 0 )
+	, bShowingConsoleForSlowTask( false )
+{
+}
+
+FMacNativeFeedbackContext::~FMacNativeFeedbackContext()
+{
+}
+
+void FMacNativeFeedbackContext::Serialize( const TCHAR* V, ELogVerbosity::Type Verbosity, const class FName& Category )
+{
+	if( !GLog->IsRedirectingTo( this ) )
+	{
+		GLog->Serialize( V, Verbosity, Category );
+	}
+}
+
+bool FMacNativeFeedbackContext::YesNof(const FText& Text)
+{
+	return GWarn->YesNof(Text);
+}
+
+bool FMacNativeFeedbackContext::ReceivedUserCancel()
+{
+	return false;
+}
+
+void FMacNativeFeedbackContext::BeginSlowTask( const FText& Task, bool bShowProgressDialog, bool bInShowCancelButton )
+{
+	GIsSlowTask = ++SlowTaskCount>0;
+
+	if(SlowTaskCount > 0 && bShowProgressDialog && GLogConsole != NULL && !GLogConsole->IsShown())
+	{
+		GLogConsole->Show(true);
+		bShowingConsoleForSlowTask = true;
+	}
+}
+
+void FMacNativeFeedbackContext::EndSlowTask()
+{
+	check(SlowTaskCount>0);
+	GIsSlowTask = --SlowTaskCount>0;
+
+	if(SlowTaskCount == 0 && bShowingConsoleForSlowTask)
+	{
+		if(GLogConsole != NULL) GLogConsole->Show(false);
+		bShowingConsoleForSlowTask = false;
+	}
+}
+
+bool FMacNativeFeedbackContext::StatusUpdate( int32 Numerator, int32 Denominator, const FText& NewStatus )
+{
+	return true;
+}
+
+bool FMacNativeFeedbackContext::StatusForceUpdate( int32 Numerator, int32 Denominator, const FText& StatusText )
+{
+	return true;
+}
+
+void FMacNativeFeedbackContext::UpdateProgress(int32 Numerator, int32 Denominator)
+{
+}
+
+FContextSupplier* FMacNativeFeedbackContext::GetContext() const
+{
+	return Context;
+}
+
+void FMacNativeFeedbackContext::SetContext( FContextSupplier* InSupplier )
+{
+	Context = InSupplier;
+}
diff --git a/Engine/Source/Developer/DesktopPlatform/Private/Mac/MacNativeFeedbackContext.h b/Engine/Source/Developer/DesktopPlatform/Private/Mac/MacNativeFeedbackContext.h
new file mode 100644
--- /dev/null
+++ b/Engine/Source/Developer/DesktopPlatform/Private/Mac/MacNativeFeedbackContext.h
@@ -0,0 +1,38 @@
+// Copyright 1998-2014 Epic Games, Inc. All Rights Reserved.
+
+/*=============================================================================
+	MacNativeFeedbackContext.h: Unreal Mac user interface interaction.
+=============================================================================*/
+
+#pragma once
+
+/**
+ * Feedback context implementation for Mac.
+ */
+class FMacNativeFeedbackContext : public FFeedbackContext
+{
+public:
+	// Constructor.
+	FMacNativeFeedbackContext();
+	virtual ~FMacNativeFeedbackContext();
+
+	virtual void Serialize( const TCHAR* V, ELogVerbosity::Type Verbosity, const class FName& Category ) OVERRIDE;
+
+	virtual bool YesNof(const FText& Text) OVERRIDE;
+
+	virtual bool ReceivedUserCancel() OVERRIDE;
+	virtual void BeginSlowTask( const FText& Task, bool bShowProgressDialog, bool bInShowCancelButton=false ) OVERRIDE;
+	virtual void EndSlowTask() OVERRIDE;
+
+	virtual bool StatusUpdate( int32 Numerator, int32 Denominator, const FText& NewStatus ) OVERRIDE;
+	virtual bool StatusForceUpdate( int32 Numerator, int32 Denominator, const FText& StatusText ) OVERRIDE;
+	virtual void UpdateProgress(int32 Numerator, int32 Denominator) OVERRIDE;
+
+	FContextSupplier* GetContext() const;
+	void SetContext( FContextSupplier* InSupplier );
+
+private:
+	FContextSupplier* Context;
+	int32 SlowTaskCount;
+	bool bShowingConsoleForSlowTask;
+};
