Clean binaries and intermediates before upgrading projects. Prevents stale build products, and allows detecting missing binaries on startup.

diff --git a/Engine/Source/Developer/DesktopPlatform/Private/DesktopPlatformBase.cpp b/Engine/Source/Developer/DesktopPlatform/Private/DesktopPlatformBase.cpp
--- a/Engine/Source/Developer/DesktopPlatform/Private/DesktopPlatformBase.cpp
+++ b/Engine/Source/Developer/DesktopPlatform/Private/DesktopPlatformBase.cpp
@@ -246,6 +246,50 @@ bool FDesktopPlatformBase::GetEngineIdentifierForProject(const FString &ProjectF
 	return false;
 }
 
+bool FDesktopPlatformBase::CleanGameProject(const FString &ProjectDir, FFeedbackContext* Warn)
+{
+	// Begin a task
+	Warn->BeginSlowTask(LOCTEXT("CleaningProject", "Removing stale build products..."), true);
+
+	// Enumerate all the files
+	TArray<FString> FileNames;
+	GetProjectBuildProducts(ProjectDir, FileNames);
+
+	// Remove all the files
+	for(int32 Idx = 0; Idx < FileNames.Num(); Idx++)
+	{
+		// Remove the file
+		if(!IFileManager::Get().Delete(*FileNames[Idx]))
+		{
+			Warn->Logf(ELogVerbosity::Error, TEXT("ERROR: Couldn't delete file '%s'"), *FileNames[Idx]);
+			return false;
+		}
+
+		// Try to remove an empty directory hierarchy
+		FString DirName = FPaths::GetPath(FileNames[Idx]);
+		for(;;)
+		{
+			TArray<FString> Contents;
+			IFileManager::Get().FindFiles(Contents, *(DirName / TEXT("*")), true, true);
+
+			if(Contents.Num() > 0) break;
+			if(!IFileManager::Get().DeleteDirectory(*DirName)) break;
+
+			int32 ParentIdx;
+			if(!DirName.FindLastChar(TEXT('/'), ParentIdx)) break;
+
+			DirName = DirName.Left(ParentIdx);
+		}
+
+		// Update the progress
+		Warn->UpdateProgress(Idx, FileNames.Num());
+	}
+
+	// End the task
+	Warn->EndSlowTask();
+	return true;
+}
+
 bool FDesktopPlatformBase::CompileGameProject(const FString& RootDir, const FString& ProjectFileName, FFeedbackContext* Warn)
 {
 	// Get the target name
@@ -451,4 +495,66 @@ const FUProjectDictionary &FDesktopPlatformBase::GetCachedProjectDictionary(cons
 	return *Dictionary;
 }
 
+void FDesktopPlatformBase::GetProjectBuildProducts(const FString& ProjectDir, TArray<FString> &OutFileNames)
+{
+	FString NormalizedProjectDir = ProjectDir;
+	FPaths::NormalizeDirectoryName(NormalizedProjectDir);
+
+	// Find all the build roots
+	TArray<FString> BuildRootDirectories;
+	BuildRootDirectories.Add(NormalizedProjectDir);
+
+	// Add all the plugin directories
+	TArray<FString> PluginFileNames;
+	IFileManager::Get().FindFilesRecursive(PluginFileNames, *(NormalizedProjectDir / TEXT("Plugins")), TEXT("*.uplugin"), true, false);
+	for(int32 Idx = 0; Idx < PluginFileNames.Num(); Idx++)
+	{
+		BuildRootDirectories.Add(FPaths::GetPath(PluginFileNames[Idx]));
+	}
+
+	// Find all the target filenames
+	TArray<FString> TargetNames;
+	IFileManager::Get().FindFilesRecursive(TargetNames, *(NormalizedProjectDir / TEXT("Source")), TEXT("*.Target.cs"), true, false, false);
+	for(int32 Idx = 0; Idx < TargetNames.Num(); Idx++)
+	{
+		TargetNames[Idx] = FPaths::GetCleanFilename(TargetNames[Idx]);
+		TargetNames[Idx].RemoveFromEnd(TEXT(".Target.cs"));
+	}
+	TargetNames.Add(TEXT("UE4Editor"));
+
+	// Add the binaries for all of the build roots
+	for(int32 RootIdx = 0; RootIdx < BuildRootDirectories.Num(); RootIdx++)
+	{
+		const FString &BuildRootDirectory = BuildRootDirectories[RootIdx];
+
+		// Find all the binaries under this build root
+		TArray<FString> Binaries;
+		IFileManager::Get().FindFilesRecursive(Binaries, *(BuildRootDirectory / TEXT("Binaries")), TEXT("*"), true, false);
+
+		// Remove all the binaries starting with a target name
+		for(int32 BinaryIdx = 0; BinaryIdx < Binaries.Num(); BinaryIdx++)
+		{
+			FString BinaryName = FPaths::GetCleanFilename(Binaries[BinaryIdx]);
+			for(int32 TargetIdx = 0; TargetIdx < TargetNames.Num(); TargetIdx++)
+			{
+				if(BinaryName.StartsWith(TargetNames[TargetIdx]))
+				{
+					const TCHAR *End = *BinaryName + TargetNames[TargetIdx].Len();
+					if(*End == 0 || *End == TEXT('.') || *End == TEXT('-'))
+					{
+						OutFileNames.Add(Binaries[BinaryIdx]);
+						break;
+					}
+				}
+			}
+		}
+	}
+
+	// Add the intermediate folders for all of the build roots
+	for(int32 RootIdx = 0; RootIdx < BuildRootDirectories.Num(); RootIdx++)
+	{
+		IFileManager::Get().FindFilesRecursive(OutFileNames, *(BuildRootDirectories[RootIdx] / TEXT("Intermediate")), TEXT("*"), true, false, false);
+	}
+}
+
 #undef LOCTEXT_NAMESPACE
diff --git a/Engine/Source/Developer/DesktopPlatform/Private/DesktopPlatformBase.h b/Engine/Source/Developer/DesktopPlatform/Private/DesktopPlatformBase.h
--- a/Engine/Source/Developer/DesktopPlatform/Private/DesktopPlatformBase.h
+++ b/Engine/Source/Developer/DesktopPlatform/Private/DesktopPlatformBase.h
@@ -31,6 +31,7 @@ public:
 	virtual bool SetEngineIdentifierForProject(const FString &ProjectFileName, const FString &Identifier) OVERRIDE;
 	virtual bool GetEngineIdentifierForProject(const FString &ProjectFileName, FString &OutIdentifier) OVERRIDE;
 
+	virtual bool CleanGameProject(const FString& ProjectDir, FFeedbackContext* Warn) OVERRIDE;
 	virtual bool CompileGameProject(const FString& RootDir, const FString& ProjectFileName, FFeedbackContext* Warn) OVERRIDE;
 	virtual bool GenerateProjectFiles(const FString& RootDir, const FString& ProjectFileName, FFeedbackContext* Warn) OVERRIDE;
 
@@ -48,4 +49,6 @@ private:
 	bool SaveProjectFile(const FString &FileName, TSharedPtr<FJsonObject> Object);
 
 	const FUProjectDictionary &GetCachedProjectDictionary(const FString& RootDir);
+
+	void GetProjectBuildProducts(const FString& ProjectFileName, TArray<FString> &OutFileNames);
 };
diff --git a/Engine/Source/Developer/DesktopPlatform/Public/IDesktopPlatform.h b/Engine/Source/Developer/DesktopPlatform/Public/IDesktopPlatform.h
--- a/Engine/Source/Developer/DesktopPlatform/Public/IDesktopPlatform.h
+++ b/Engine/Source/Developer/DesktopPlatform/Public/IDesktopPlatform.h
@@ -257,6 +257,14 @@ public:
 	*/
 	virtual bool GetEngineIdentifierForProject(const FString &ProjectFileName, FString &OutIdentifier) = 0;
 
+	/**
+	* Cleans a game project. Removes the intermediate folder and binary build products.
+	*
+	* @param ProjectDirName		Directory for the project
+	* @param OutFileNames		Output array of the project's build products
+	*/
+	virtual bool CleanGameProject(const FString& ProjectDir, FFeedbackContext* Warn) = 0;
+
 	/**
 	* Compiles a game project.
 	*
diff --git a/Engine/Source/Editor/GameProjectGeneration/Private/SProjectBrowser.cpp b/Engine/Source/Editor/GameProjectGeneration/Private/SProjectBrowser.cpp
--- a/Engine/Source/Editor/GameProjectGeneration/Private/SProjectBrowser.cpp
+++ b/Engine/Source/Editor/GameProjectGeneration/Private/SProjectBrowser.cpp
@@ -898,22 +898,31 @@ bool SProjectBrowser::OpenProject( const FString& InProjectFile )
 		// Ask whether to copy the project or upgrade in-place
 		TArray<FText> Buttons;
 		int32 OpenCopyButton = Buttons.Add(LOCTEXT("ProjectUpgradeCopy", "Create a copy"));
-		Buttons.Add(LOCTEXT("ProjectUpgradeInPlace", "Convert in-place"));
+		int32 InPlaceButton = Buttons.Add(LOCTEXT("ProjectUpgradeInPlace", "Convert in-place"));
 		int32 CancelButton = Buttons.Add(LOCTEXT("ProjectUpgradeCancel", "Cancel"));
 
 		// Show the dialog
-		int32 Selection = SVerbChoiceDialog::ShowModal(LOCTEXT("ProjectUpgradeTitle", "Upgrade Project"), LOCTEXT("ProjectUpgradePrompt", "This project was made with a different version of the editor. Converting it to this version of the editor may be irreversible.\n\nWould you like to create a copy before opening?"), Buttons);
+		int32 Selection = SVerbChoiceDialog::ShowModal(LOCTEXT("ProjectConversionTitle", "Convert Project"), LOCTEXT("ConvertProjectPrompt", "This project was made with a different version of the editor. Converting to this version may be irreversible or lose data.\n\nIt is recommened that you create a copy before proceeding."), Buttons);
 		if(Selection == OpenCopyButton)
 		{
 			FString NewProjectFile;
 			if(!GameProjectUtils::DuplicateProjectForUpgrade(ProjectFile, NewProjectFile))
 			{
-				FMessageDialog::Open( EAppMsgType::Ok, LOCTEXT("CopyFailed", "Couldn't copy project. Check you have sufficient hard drive space and write access to the project folder.") );
+				FMessageDialog::Open( EAppMsgType::Ok, LOCTEXT("ConvertProjectCopyFailed", "Couldn't copy project. Check you have sufficient hard drive space and write access to the project folder.") );
 				return false;
 			}
 			ProjectFile = NewProjectFile;
 		}
-		else if(Selection == CancelButton)
+		else if(Selection == InPlaceButton)
+		{
+			FString ProjectDir = FPaths::GetPath(ProjectFile);
+			if(!FDesktopPlatformModule::Get()->CleanGameProject(ProjectDir, GWarn))
+			{
+				FMessageDialog::Open( EAppMsgType::Ok, LOCTEXT("ConvertProjectCleanFailed", "Couldn't clean project binary and intermediate folders for upgrade. Check that the project folders are writable."));
+				return false;
+			}
+		}
+		else 
 		{
 			return false;
 		}
diff --git a/Engine/Source/Programs/UnrealVersionSelector/Private/UnrealVersionSelector.cpp b/Engine/Source/Programs/UnrealVersionSelector/Private/UnrealVersionSelector.cpp
--- a/Engine/Source/Programs/UnrealVersionSelector/Private/UnrealVersionSelector.cpp
+++ b/Engine/Source/Programs/UnrealVersionSelector/Private/UnrealVersionSelector.cpp
@@ -75,6 +75,14 @@ bool SwitchVersion(const FString& ProjectFileName)
 		return false;
 	}
 
+	// Clean all the build products
+	FFeedbackContext *Warn = FDesktopPlatformModule::Get()->GetNativeFeedbackContext();
+	if (!FDesktopPlatformModule::Get()->CleanGameProject(FPaths::GetPath(ProjectFileName), Warn))
+	{
+		FPlatformMisc::MessageBoxExt(EAppMsgType::Ok, TEXT("Couldn't clean project build products."), TEXT("Error"));
+		return false;
+	}
+
 	// Generate project files
 	return GenerateProjectFiles(ProjectFileName);
 }
