Fix updates not writing out to the log in FWindowsNativeFeedbackContext.

diff --git a/Engine/Source/Developer/DesktopPlatform/Private/Windows/WindowsNativeFeedbackContext.cpp b/Engine/Source/Developer/DesktopPlatform/Private/Windows/WindowsNativeFeedbackContext.cpp
--- a/Engine/Source/Developer/DesktopPlatform/Private/Windows/WindowsNativeFeedbackContext.cpp
+++ b/Engine/Source/Developer/DesktopPlatform/Private/Windows/WindowsNativeFeedbackContext.cpp
@@ -75,6 +75,7 @@ void FWindowsNativeFeedbackContext::Serialize( const TCHAR* V, ELogVerbosity::Ty
 		if(hThread != NULL)
 		{
 			LogOutput += FString(V) + FString("\r\n");
+			SetEvent(hUpdateEvent);
 		}
 	}
 }
