Change detection of default engine for projects; just look for any valid engine directory up the hierarchy.

diff --git a/Engine/Source/Developer/DesktopPlatform/Private/DesktopPlatformBase.cpp b/Engine/Source/Developer/DesktopPlatform/Private/DesktopPlatformBase.cpp
--- a/Engine/Source/Developer/DesktopPlatform/Private/DesktopPlatformBase.cpp
+++ b/Engine/Source/Developer/DesktopPlatform/Private/DesktopPlatformBase.cpp
@@ -215,17 +215,16 @@ bool FDesktopPlatformBase::GetEngineIdentifierForProject(const FString &ProjectF
 		return true;
 	}
 
-	// Otherwise scan up through the directory hierarchy to find an installation which references it through .uprojectdirs
+	// Otherwise scan up through the directory hierarchy to find an installation
 	FString ParentDir = FPaths::GetPath(ProjectFileName);
 	FPaths::NormalizeDirectoryName(ParentDir);
 
+	// Keep going until we reach the root
 	int32 SeparatorIdx;
 	while(ParentDir.FindLastChar(TEXT('/'), SeparatorIdx))
 	{
 		ParentDir.RemoveAt(SeparatorIdx, ParentDir.Len() - SeparatorIdx);
-
-		const FUProjectDictionary &Dictionary = GetCachedProjectDictionary(ParentDir);
-		if(!Dictionary.IsForeignProject(ProjectFileName) && GetEngineIdentifierFromRootDir(ParentDir, OutIdentifier))
+		if(IsValidRootDirectory(ParentDir) && GetEngineIdentifierFromRootDir(ParentDir, OutIdentifier))
 		{
 			return true;
 		}
