Fix generating project files for non-foreign projects.

diff --git a/Engine/Source/Developer/DesktopPlatform/Private/DesktopPlatformBase.cpp b/Engine/Source/Developer/DesktopPlatform/Private/DesktopPlatformBase.cpp
--- a/Engine/Source/Developer/DesktopPlatform/Private/DesktopPlatformBase.cpp
+++ b/Engine/Source/Developer/DesktopPlatform/Private/DesktopPlatformBase.cpp
@@ -333,19 +333,19 @@ bool FDesktopPlatformBase::GenerateProjectFiles(const FString& RootDir, const FS
 		if(ProjectDictionary.IsForeignProject(ProjectFileName))
 		{
 			Arguments += FString::Printf(TEXT(" -project=\"%s\""), *IFileManager::Get().ConvertToAbsolutePathForExternalAppForRead(*ProjectFileName));
-		}
 
-		// Always include game source
-		Arguments += " -game";
+			// Always include game source
+			Arguments += " -game";
 
-		// Determine whether or not to include engine source
-		if(IsSourceDistribution(RootDir))
-		{
-			Arguments += " -engine";
-		}
-		else
-		{
-			Arguments += " -rocket";
+			// Determine whether or not to include engine source
+			if(IsSourceDistribution(RootDir))
+			{
+				Arguments += " -engine";
+			}
+			else
+			{
+				Arguments += " -rocket";
+			}
 		}
 	}
 	Arguments += " -progress";
diff --git a/Engine/Source/Developer/DesktopPlatform/Private/Windows/DesktopPlatformWindows.cpp b/Engine/Source/Developer/DesktopPlatform/Private/Windows/DesktopPlatformWindows.cpp
--- a/Engine/Source/Developer/DesktopPlatform/Private/Windows/DesktopPlatformWindows.cpp
+++ b/Engine/Source/Developer/DesktopPlatform/Private/Windows/DesktopPlatformWindows.cpp
@@ -478,6 +478,9 @@ bool FDesktopPlatformWindows::RunUnrealBuildTool(const FText& Description, const
 		return false;
 	}
 
+	// Write the output
+	Warn->Logf(TEXT("Running %s %s"), *UnrealBuildToolPath, *Arguments);
+
 	// Spawn UBT
 	int32 ExitCode = 0;
 	return FFeedbackContextMarkup::PipeProcessOutput(Description, UnrealBuildToolPath, Arguments, Warn, &ExitCode) && ExitCode == 0;
