Move the feedback context window to the front when it's first displayed.

diff --git a/Engine/Source/Developer/DesktopPlatform/Private/Windows/WindowsNativeFeedbackContext.cpp b/Engine/Source/Developer/DesktopPlatform/Private/Windows/WindowsNativeFeedbackContext.cpp
--- a/Engine/Source/Developer/DesktopPlatform/Private/Windows/WindowsNativeFeedbackContext.cpp
+++ b/Engine/Source/Developer/DesktopPlatform/Private/Windows/WindowsNativeFeedbackContext.cpp
@@ -279,6 +279,8 @@ DWORD FWindowsNativeFeedbackContext::SlowTaskThreadProc(void* ThreadParam)
 	ShowWindow(hWnd, SW_SHOW);
 	UpdateWindow(hWnd);
 
+	SetForegroundWindow(hWnd);
+
 	FString PrevStatus;
 	float PrevProgress = 0.0f;
 	int32 PrevLogOutputLength = 0;
