[INTEGRATE] Everything from 4.2 branch back into main.

diff --git a/Engine/Plugins/Slate/SlateRemote/Source/SlateRemote/Private/SlateRemoteModule.cpp b/Engine/Plugins/Slate/SlateRemote/Source/SlateRemote/Private/SlateRemoteModule.cpp
--- a/Engine/Plugins/Slate/SlateRemote/Source/SlateRemote/Private/SlateRemoteModule.cpp
+++ b/Engine/Plugins/Slate/SlateRemote/Source/SlateRemote/Private/SlateRemoteModule.cpp
@@ -126,7 +126,7 @@ protected:
 		}
 
 		// create server
-		ISocketSubsystem* SocketSubsystem = ISocketSubsystem::Get();
+		ISocketSubsystem* SocketSubsystem = ISocketSubsystem::Get(PLATFORM_SOCKETSUBSYSTEM);
 
 		if (SocketSubsystem != nullptr)
 		{
diff --git a/Engine/Source/Developer/DesktopPlatform/Private/DesktopPlatformBase.cpp b/Engine/Source/Developer/DesktopPlatform/Private/DesktopPlatformBase.cpp
--- a/Engine/Source/Developer/DesktopPlatform/Private/DesktopPlatformBase.cpp
+++ b/Engine/Source/Developer/DesktopPlatform/Private/DesktopPlatformBase.cpp
@@ -252,7 +252,8 @@ bool FDesktopPlatformBase::CleanGameProject(const FString &ProjectDir, FFeedback
 
 	// Enumerate all the files
 	TArray<FString> FileNames;
-	GetProjectBuildProducts(ProjectDir, FileNames);
+	TArray<FString> DirectoryNames;
+	GetProjectBuildProducts(ProjectDir, FileNames, DirectoryNames);
 
 	// Remove all the files
 	for(int32 Idx = 0; Idx < FileNames.Num(); Idx++)
@@ -264,24 +265,22 @@ bool FDesktopPlatformBase::CleanGameProject(const FString &ProjectDir, FFeedback
 			return false;
 		}
 
-		// Try to remove an empty directory hierarchy
-		FString DirName = FPaths::GetPath(FileNames[Idx]);
-		for(;;)
-		{
-			TArray<FString> Contents;
-			IFileManager::Get().FindFiles(Contents, *(DirName / TEXT("*")), true, true);
-
-			if(Contents.Num() > 0) break;
-			if(!IFileManager::Get().DeleteDirectory(*DirName)) break;
-
-			int32 ParentIdx;
-			if(!DirName.FindLastChar(TEXT('/'), ParentIdx)) break;
+		// Update the progress
+		Warn->UpdateProgress(Idx, FileNames.Num() + DirectoryNames.Num());
+	}
 
-			DirName = DirName.Left(ParentIdx);
+	// Remove all the directories
+	for(int32 Idx = 0; Idx < DirectoryNames.Num(); Idx++)
+	{
+		// Remove the directory
+		if(!IFileManager::Get().DeleteDirectory(*DirectoryNames[Idx], false, true))
+		{
+			Warn->Logf(ELogVerbosity::Error, TEXT("ERROR: Couldn't delete directory '%s'"), *DirectoryNames[Idx]);
+			return false;
 		}
 
 		// Update the progress
-		Warn->UpdateProgress(Idx, FileNames.Num());
+		Warn->UpdateProgress(Idx + FileNames.Num(), FileNames.Num() + DirectoryNames.Num());
 	}
 
 	// End the task
@@ -494,7 +493,7 @@ const FUProjectDictionary &FDesktopPlatformBase::GetCachedProjectDictionary(cons
 	return *Dictionary;
 }
 
-void FDesktopPlatformBase::GetProjectBuildProducts(const FString& ProjectDir, TArray<FString> &OutFileNames)
+void FDesktopPlatformBase::GetProjectBuildProducts(const FString& ProjectDir, TArray<FString> &OutFileNames, TArray<FString> &OutDirectoryNames)
 {
 	FString NormalizedProjectDir = ProjectDir;
 	FPaths::NormalizeDirectoryName(NormalizedProjectDir);
@@ -511,48 +510,16 @@ void FDesktopPlatformBase::GetProjectBuildProducts(const FString& ProjectDir, TA
 		BuildRootDirectories.Add(FPaths::GetPath(PluginFileNames[Idx]));
 	}
 
-	// Find all the target filenames
-	TArray<FString> TargetNames;
-	IFileManager::Get().FindFilesRecursive(TargetNames, *(NormalizedProjectDir / TEXT("Source")), TEXT("*.Target.cs"), true, false, false);
-	for(int32 Idx = 0; Idx < TargetNames.Num(); Idx++)
-	{
-		TargetNames[Idx] = FPaths::GetCleanFilename(TargetNames[Idx]);
-		TargetNames[Idx].RemoveFromEnd(TEXT(".Target.cs"));
-	}
-	TargetNames.Add(TEXT("UE4Editor"));
-
-	// Add the binaries for all of the build roots
-	for(int32 RootIdx = 0; RootIdx < BuildRootDirectories.Num(); RootIdx++)
+	// Add all the intermediate directories
+	for(int32 Idx = 0; Idx < BuildRootDirectories.Num(); Idx++)
 	{
-		const FString &BuildRootDirectory = BuildRootDirectories[RootIdx];
-
-		// Find all the binaries under this build root
-		TArray<FString> Binaries;
-		IFileManager::Get().FindFilesRecursive(Binaries, *(BuildRootDirectory / TEXT("Binaries")), TEXT("*"), true, false);
-
-		// Remove all the binaries starting with a target name
-		for(int32 BinaryIdx = 0; BinaryIdx < Binaries.Num(); BinaryIdx++)
-		{
-			FString BinaryName = FPaths::GetCleanFilename(Binaries[BinaryIdx]);
-			for(int32 TargetIdx = 0; TargetIdx < TargetNames.Num(); TargetIdx++)
-			{
-				if(BinaryName.StartsWith(TargetNames[TargetIdx]))
-				{
-					const TCHAR *End = *BinaryName + TargetNames[TargetIdx].Len();
-					if(*End == 0 || *End == TEXT('.') || *End == TEXT('-'))
-					{
-						OutFileNames.Add(Binaries[BinaryIdx]);
-						break;
-					}
-				}
-			}
-		}
+		OutDirectoryNames.Add(BuildRootDirectories[Idx] / TEXT("Intermediate"));
 	}
 
-	// Add the intermediate folders for all of the build roots
-	for(int32 RootIdx = 0; RootIdx < BuildRootDirectories.Num(); RootIdx++)
+	// Add the files in the cleaned directories to the output list
+	for(int32 Idx = 0; Idx < OutDirectoryNames.Num(); Idx++)
 	{
-		IFileManager::Get().FindFilesRecursive(OutFileNames, *(BuildRootDirectories[RootIdx] / TEXT("Intermediate")), TEXT("*"), true, false, false);
+		IFileManager::Get().FindFilesRecursive(OutFileNames, *OutDirectoryNames[Idx], TEXT("*"), true, false, false);
 	}
 }
 
diff --git a/Engine/Source/Developer/DesktopPlatform/Private/DesktopPlatformBase.h b/Engine/Source/Developer/DesktopPlatform/Private/DesktopPlatformBase.h
--- a/Engine/Source/Developer/DesktopPlatform/Private/DesktopPlatformBase.h
+++ b/Engine/Source/Developer/DesktopPlatform/Private/DesktopPlatformBase.h
@@ -50,5 +50,5 @@ private:
 
 	const FUProjectDictionary &GetCachedProjectDictionary(const FString& RootDir);
 
-	void GetProjectBuildProducts(const FString& ProjectFileName, TArray<FString> &OutFileNames);
+	void GetProjectBuildProducts(const FString& ProjectFileName, TArray<FString> &OutFileNames, TArray<FString> &OutDirectoryNames);
 };
diff --git a/Engine/Source/Editor/GameProjectGeneration/Private/SProjectBrowser.cpp b/Engine/Source/Editor/GameProjectGeneration/Private/SProjectBrowser.cpp
--- a/Engine/Source/Editor/GameProjectGeneration/Private/SProjectBrowser.cpp
+++ b/Engine/Source/Editor/GameProjectGeneration/Private/SProjectBrowser.cpp
@@ -983,6 +983,14 @@ bool SProjectBrowser::OpenProject( const FString& InProjectFile )
 			}
 			ProjectFile = NewProjectFile;
 		}
+		if(Selection == OpenExistingButton)
+		{
+			if(!FDesktopPlatformModule::Get()->CleanGameProject(FPaths::GetPath(ProjectFile), GWarn))
+			{
+				FMessageDialog::Open( EAppMsgType::Ok, LOCTEXT("ConvertProjectCleanFailed", "Couldn't clean project build files. Delete the project intermediate folder and try again.") );
+				return false;
+			}
+		}
 		if(Selection != SkipConversionButton)
 		{
 			// If it's a code-based project, generate project files and open visual studio after an upgrade
diff --git a/Engine/Source/Programs/AutomationTool/Scripts/GUBP.Automation.cs b/Engine/Source/Programs/AutomationTool/Scripts/GUBP.Automation.cs
--- a/Engine/Source/Programs/AutomationTool/Scripts/GUBP.Automation.cs
+++ b/Engine/Source/Programs/AutomationTool/Scripts/GUBP.Automation.cs
@@ -1139,6 +1139,36 @@ public class GUBP : BuildCommand
                 AddPseudodependency(WaitForTestShared.StaticGetFullName());
                 AgentSharingGroup = "TemplateMonolithics" + StaticGetHostPlatformSuffix(InHostPlatform);
             }
+
+            if (GUBP.bBuildRocket)
+            {
+                var Target = bp.Branch.BaseEngineProject.Properties.Targets[TargetRules.TargetType.Game];
+                var Configs = Target.Rules.GUBP_GetConfigs_MonolithicOnly(HostPlatform, TargetPlatform);
+                foreach (var Config in Configs)
+                {
+                    if (HostPlatform == UnrealTargetPlatform.Win64)
+                    {
+                        if (TargetPlatform == UnrealTargetPlatform.Win32 && Config != UnrealTargetConfiguration.Shipping)
+                        {
+                            continue;
+                        }
+                        if (TargetPlatform == UnrealTargetPlatform.Win64 && Config != UnrealTargetConfiguration.Development)
+                        {
+                            continue;
+                        }
+                        if (TargetPlatform == UnrealTargetPlatform.Android && Config != UnrealTargetConfiguration.Shipping && Config != UnrealTargetConfiguration.Development)
+                        {
+                            continue;
+                        }
+                    }
+                    else if (Config != UnrealTargetConfiguration.Shipping && Config != UnrealTargetConfiguration.Development)
+                    {
+                        continue;
+                    }
+                    Log("Building {0} for Host={1} Target={2} Config={3} for rocket.   Node={4}", Target.TargetName, HostPlatform, TargetPlatform, Config, GetFullName());
+                }
+            }
+
         }
 
         public static string StaticGetFullName(UnrealTargetPlatform InHostPlatform, BranchInfo.BranchUProject InGameProj, UnrealTargetPlatform InTargetPlatform)
@@ -1198,11 +1228,22 @@ public class GUBP : BuildCommand
                         {
                             if (GUBP.bBuildRocket)
                             {
-                                if (HostPlatform == UnrealTargetPlatform.Win64 && TargetPlatform == UnrealTargetPlatform.Win32 && Config != UnrealTargetConfiguration.Shipping)
+                                if (HostPlatform == UnrealTargetPlatform.Win64)
                                 {
-                                    continue;
+                                    if (TargetPlatform == UnrealTargetPlatform.Win32 && Config != UnrealTargetConfiguration.Shipping)
+                                    {
+                                        continue;
+                                    }
+                                    if (TargetPlatform == UnrealTargetPlatform.Win64 && Config != UnrealTargetConfiguration.Development)
+                                    {
+                                        continue;
+                                    }
+                                    if (TargetPlatform == UnrealTargetPlatform.Android && Config != UnrealTargetConfiguration.Shipping && Config != UnrealTargetConfiguration.Development)
+                                    {
+                                        continue;
+                                    }
                                 }
-                                if (HostPlatform == UnrealTargetPlatform.Win64 && TargetPlatform == UnrealTargetPlatform.Win64 && Config != UnrealTargetConfiguration.Development)
+                                else if (Config != UnrealTargetConfiguration.Shipping && Config != UnrealTargetConfiguration.Development)
                                 {
                                     continue;
                                 }
diff --git a/Engine/Source/Runtime/Engine/Private/Slate/SlateGameResources.cpp b/Engine/Source/Runtime/Engine/Private/Slate/SlateGameResources.cpp
--- a/Engine/Source/Runtime/Engine/Private/Slate/SlateGameResources.cpp
+++ b/Engine/Source/Runtime/Engine/Private/Slate/SlateGameResources.cpp
@@ -124,7 +124,7 @@ const FSlateWidgetStyle* FSlateGameResources::GetWidgetStyleInternal( const FNam
 
 	const FSlateWidgetStyle* Style = StyleAsset->GetStyleChecked( DesiredTypeName );
 
-	if ( Style == NULL )
+	if ( Style == NULL && GIsEditor)
 	{
 		TSharedRef< FTokenizedMessage > Message = FTokenizedMessage::Create( EMessageSeverity::Error, FText::Format( NSLOCTEXT("SlateStyleSet", "WrongWidgetStyleType", "The Slate Widget Style '{0}' is not of the desired type. Desired: '{1}', Actual: '{2}'"), FText::FromName( StyleName ), FText::FromName( DesiredTypeName ), FText::FromName( StyleAsset->CustomStyle->GetStyle()->GetTypeName() ) ) );
 		Message->AddToken( FAssetNameToken::Create( StyleAsset->GetPathName(), FText::FromString( StyleAsset->GetName() ) ) );
diff --git a/Engine/Source/Runtime/Launch/Private/Android/LaunchAndroid.cpp b/Engine/Source/Runtime/Launch/Private/Android/LaunchAndroid.cpp
--- a/Engine/Source/Runtime/Launch/Private/Android/LaunchAndroid.cpp
+++ b/Engine/Source/Runtime/Launch/Private/Android/LaunchAndroid.cpp
@@ -242,6 +242,7 @@ int32 AndroidMain(struct android_app* state)
 		IgnoredGamepadKeyCodes.Add(IgnoredGamepadKeyCodesList[i]);
 	}
 
+	// wait a moment for the data from the activity to be set
 	FPlatformProcess::Sleep(1.0f);
 
 	// read the command line file
