Add a single codepath to enumerate projects for an engine installation, which can be shared with the launcher.

diff --git a/Engine/Source/Developer/DesktopPlatform/Private/DesktopPlatformBase.cpp b/Engine/Source/Developer/DesktopPlatform/Private/DesktopPlatformBase.cpp
--- a/Engine/Source/Developer/DesktopPlatform/Private/DesktopPlatformBase.cpp
+++ b/Engine/Source/Developer/DesktopPlatform/Private/DesktopPlatformBase.cpp
@@ -5,6 +5,7 @@
 #include "UProjectInfo.h"
 #include "EngineVersion.h"
 #include "ModuleManager.h"
+#include "Runtime/Launch/Resources/Version.h"
 
 #define LOCTEXT_NAMESPACE "DesktopPlatform"
 
@@ -54,6 +55,21 @@ void FDesktopPlatformBase::EnumerateLauncherSampleInstallations(TArray<FString>
 	}
 }
 
+void FDesktopPlatformBase::EnumerateLauncherSampleProjects(TArray<FString> &OutFileNames)
+{
+	// Enumerate all the sample installation directories
+	TArray<FString> LauncherSampleDirectories;
+	EnumerateLauncherSampleInstallations(LauncherSampleDirectories);
+
+	// Find all the project files within them
+	for(int32 Idx = 0; Idx < LauncherSampleDirectories.Num(); Idx++)
+	{
+		TArray<FString> FileNames;
+		IFileManager::Get().FindFiles(FileNames, *(LauncherSampleDirectories[Idx] / TEXT("*.uproject")), true, false);
+		OutFileNames.Append(FileNames);
+	}
+}
+
 bool FDesktopPlatformBase::GetEngineRootDirFromIdentifier(const FString &Identifier, FString &OutRootDir)
 {
 	// Get all the installations
@@ -356,6 +372,13 @@ bool FDesktopPlatformBase::GenerateProjectFiles(const FString& RootDir, const FS
 	return bRes;
 }
 
+FString FDesktopPlatformBase::GetDefaultProjectCreationPath()
+{
+	// My Documents
+	const FString DefaultProjectSubFolder = TEXT("Unreal Projects");
+	return FString(FPlatformProcess::UserDir()) + DefaultProjectSubFolder;
+}
+
 void FDesktopPlatformBase::ReadLauncherInstallationList()
 {
 	FString InstalledListFile = FString(FPlatformProcess::ApplicationSettingsDir()) / TEXT("UnrealEngineLauncher/LauncherInstalled.dat");
@@ -523,4 +546,101 @@ void FDesktopPlatformBase::GetProjectBuildProducts(const FString& ProjectDir, TA
 	}
 }
 
+bool FDesktopPlatformBase::EnumerateProjectsKnownByEngine(const FString &Identifier, bool bIncludeNativeProjects, TArray<FString> &OutProjectFileNames)
+{
+	// Get the engine root directory
+	FString RootDir;
+	if(!GetEngineRootDirFromIdentifier(Identifier, RootDir))
+	{
+		return false;
+	}
+
+	// Get the path to the game agnostic settings
+	FString UserDir;
+	if(IsStockEngineRelease(Identifier))
+	{
+		UserDir = FPaths::Combine(FPlatformProcess::UserSettingsDir(), *FString(EPIC_PRODUCT_IDENTIFIER), *Identifier);
+	}
+	else
+	{
+		UserDir = FPaths::Combine(*RootDir, TEXT("Engine"));
+	}
+
+	// Get the game agnostic config dir
+	FString GameAgnosticConfigDir = UserDir / TEXT("Saved/Config") / ANSI_TO_TCHAR(FPlatformProperties::PlatformName());
+
+	// Find all the created project directories. Start with the default project creation path.
+	TArray<FString> SearchDirectories;
+	SearchDirectories.AddUnique(GetDefaultProjectCreationPath());
+
+	// Load the config file
+	FConfigFile GameAgnosticConfig;
+	FConfigCacheIni::LoadExternalIniFile(GameAgnosticConfig, TEXT("EditorGameAgnostic"), NULL, *GameAgnosticConfigDir, false);
+
+	// Find the editor game-agnostic settings
+	FConfigSection* Section = GameAgnosticConfig.Find(TEXT("/Script/UnrealEd.EditorGameAgnosticSettings"));
+	if(Section != NULL)
+	{
+		// Add in every path that the user has ever created a project file. This is to catch new projects showing up in the user's project folders
+		TArray<FString> AdditionalDirectories;
+		Section->MultiFind(TEXT("CreatedProjectPaths"), AdditionalDirectories);
+		for(int Idx = 0; Idx < AdditionalDirectories.Num(); Idx++)
+		{
+			FPaths::NormalizeDirectoryName(AdditionalDirectories[Idx]);
+			SearchDirectories.AddUnique(AdditionalDirectories[Idx]);
+		}
+
+		// Also add in all the recently opened projects
+		TArray<FString> RecentlyOpenedFiles;
+		Section->MultiFind(TEXT("RecentlyOpenedProjectFiles"), RecentlyOpenedFiles);
+		for(int Idx = 0; Idx < RecentlyOpenedFiles.Num(); Idx++)
+		{
+			FPaths::NormalizeFilename(RecentlyOpenedFiles[Idx]);
+			OutProjectFileNames.AddUnique(RecentlyOpenedFiles[Idx]);
+		}		
+	}
+
+	// Find all the other projects that are in the search directories
+	for(int Idx = 0; Idx < SearchDirectories.Num(); Idx++)
+	{
+		TArray<FString> ProjectFolders;
+		IFileManager::Get().FindFiles(ProjectFolders, *(SearchDirectories[Idx] / TEXT("*")), false, true);
+
+		for(int32 FolderIdx = 0; FolderIdx < ProjectFolders.Num(); FolderIdx++)
+		{
+			TArray<FString> ProjectFiles;
+			IFileManager::Get().FindFiles(ProjectFiles, *(ProjectFolders[FolderIdx] / TEXT("*.uproject")), true, false);
+
+			for(int32 FileIdx = 0; FileIdx < ProjectFiles.Num(); FileIdx++)
+			{
+				OutProjectFileNames.AddUnique(ProjectFiles[FileIdx]);
+			}
+		}
+	}
+
+	// Find all the native projects, and either add or remove them from the list depending on whether we want native projects
+	const FUProjectDictionary &Dictionary = GetCachedProjectDictionary(RootDir);
+	if(bIncludeNativeProjects)
+	{
+		TArray<FString> NativeProjectPaths = Dictionary.GetProjectPaths();
+		for(int Idx = 0; Idx < NativeProjectPaths.Num(); Idx++)
+		{
+			if(!NativeProjectPaths[Idx].Contains(TEXT("/Templates/")))
+			{
+				OutProjectFileNames.AddUnique(NativeProjectPaths[Idx]);
+			}
+		}
+	}
+	else
+	{
+		TArray<FString> NativeProjectPaths = Dictionary.GetProjectPaths();
+		for(int Idx = 0; Idx < NativeProjectPaths.Num(); Idx++)
+		{
+			OutProjectFileNames.Remove(NativeProjectPaths[Idx]);
+		}
+	}
+
+	return true;
+}
+
 #undef LOCTEXT_NAMESPACE
diff --git a/Engine/Source/Developer/DesktopPlatform/Private/DesktopPlatformBase.h b/Engine/Source/Developer/DesktopPlatform/Private/DesktopPlatformBase.h
--- a/Engine/Source/Developer/DesktopPlatform/Private/DesktopPlatformBase.h
+++ b/Engine/Source/Developer/DesktopPlatform/Private/DesktopPlatformBase.h
@@ -16,6 +16,7 @@ public:
 
 	virtual void EnumerateLauncherEngineInstallations(TMap<FString, FString> &OutInstallations) OVERRIDE;
 	virtual void EnumerateLauncherSampleInstallations(TArray<FString> &OutInstallations) OVERRIDE;
+	virtual void EnumerateLauncherSampleProjects(TArray<FString> &OutFileNames) OVERRIDE;
 	virtual bool GetEngineRootDirFromIdentifier(const FString &Identifier, FString &OutRootDir) OVERRIDE;
 	virtual bool GetEngineIdentifierFromRootDir(const FString &RootDir, FString &OutIdentifier) OVERRIDE;
 
@@ -35,6 +36,9 @@ public:
 	virtual bool CompileGameProject(const FString& RootDir, const FString& ProjectFileName, FFeedbackContext* Warn) OVERRIDE;
 	virtual bool GenerateProjectFiles(const FString& RootDir, const FString& ProjectFileName, FFeedbackContext* Warn) OVERRIDE;
 
+	virtual bool EnumerateProjectsKnownByEngine(const FString &Identifier, bool bIncludeNativeProjects, TArray<FString> &OutProjectFileNames) OVERRIDE;
+	virtual FString GetDefaultProjectCreationPath() OVERRIDE;
+
 private:
 	FString CurrentEngineIdentifier;
 	FDateTime LauncherInstallationTimestamp;
diff --git a/Engine/Source/Developer/DesktopPlatform/Public/IDesktopPlatform.h b/Engine/Source/Developer/DesktopPlatform/Public/IDesktopPlatform.h
--- a/Engine/Source/Developer/DesktopPlatform/Public/IDesktopPlatform.h
+++ b/Engine/Source/Developer/DesktopPlatform/Public/IDesktopPlatform.h
@@ -155,6 +155,13 @@ public:
 	*/
 	virtual void EnumerateLauncherSampleInstallations(TArray<FString> &OutInstallations) = 0;
 
+	/**
+	* Enumerates all the sample projects installed by the launcher.
+	*
+	* @param	OutFileNames		Array of sample project filenames.
+	*/
+	virtual void EnumerateLauncherSampleProjects(TArray<FString> &OutFileNames) = 0;
+
 	/**
 	* Returns the identifier for the engine with the given root directory.
 	*
@@ -302,4 +309,20 @@ public:
 	* @return FFeedbackContext for the native GUI.
 	*/
 	virtual FFeedbackContext* GetNativeFeedbackContext() = 0;
+
+	/**
+	* Finds all the projects which the engine (given by an identifier) has a record of. This includes all the recently opened projects, projects in the default project directory, and so on.
+	* 
+	* @param Identifier		Identifier for the engine 
+	* @param bIncludeNat
+	* @return Path to the folder
+	*/
+	virtual bool EnumerateProjectsKnownByEngine(const FString &Identifier, bool bIncludeNativeProjects, TArray<FString> &OutProjectFileNames) = 0;
+
+	/**
+	* Gets the default folder for creating new projects.
+	* 
+	* @return Path to the folder
+	*/
+	virtual FString GetDefaultProjectCreationPath() = 0;
 };
diff --git a/Engine/Source/Editor/GameProjectGeneration/Private/GameProjectUtils.cpp b/Engine/Source/Editor/GameProjectGeneration/Private/GameProjectUtils.cpp
--- a/Engine/Source/Editor/GameProjectGeneration/Private/GameProjectUtils.cpp
+++ b/Engine/Source/Editor/GameProjectGeneration/Private/GameProjectUtils.cpp
@@ -847,13 +847,6 @@ UTemplateProjectDefs* GameProjectUtils::LoadTemplateDefs(const FString& ProjectD
 	return TemplateDefs;
 }
 
-FString GameProjectUtils::GetDefaultProjectCreationPath()
-{
-	// My Documents
-	const FString DefaultProjectSubFolder = TEXT("Unreal Projects");
-	return FString(FPlatformProcess::UserDir()) + DefaultProjectSubFolder;
-}
-
 bool GameProjectUtils::GenerateProjectFromScratch(const FString& NewProjectFile, bool bShouldGenerateCode, bool bCopyStarterContent, FText& OutFailReason)
 {
 	const FString NewProjectFolder = FPaths::GetPath(NewProjectFile);
diff --git a/Engine/Source/Editor/GameProjectGeneration/Private/GameProjectUtils.h b/Engine/Source/Editor/GameProjectGeneration/Private/GameProjectUtils.h
--- a/Engine/Source/Editor/GameProjectGeneration/Private/GameProjectUtils.h
+++ b/Engine/Source/Editor/GameProjectGeneration/Private/GameProjectUtils.h
@@ -156,9 +156,6 @@ public:
 	/** Loads a template project definitions object from the TemplateDefs.ini file in the specified project */
 	static UTemplateProjectDefs* LoadTemplateDefs(const FString& ProjectDirectory);
 
-	/** Returns the default base folder to create a new project */
-	static FString GetDefaultProjectCreationPath();
-
 	/** Returns number of code files in the currently loaded project */
 	static int32 GetProjectCodeFileCount();
 
diff --git a/Engine/Source/Editor/GameProjectGeneration/Private/SNewProjectWizard.cpp b/Engine/Source/Editor/GameProjectGeneration/Private/SNewProjectWizard.cpp
--- a/Engine/Source/Editor/GameProjectGeneration/Private/SNewProjectWizard.cpp
+++ b/Engine/Source/Editor/GameProjectGeneration/Private/SNewProjectWizard.cpp
@@ -997,7 +997,7 @@ void SNewProjectWizard::SetDefaultProjectLocation( )
 	if ( DefaultProjectFilePath.IsEmpty() )
 	{
 		// No previously used path, decide a default path.
-		DefaultProjectFilePath = GameProjectUtils::GetDefaultProjectCreationPath();
+		DefaultProjectFilePath = FDesktopPlatformModule::Get()->GetDefaultProjectCreationPath();
 		IFileManager::Get().MakeDirectory(*DefaultProjectFilePath, true);
 	}
 
diff --git a/Engine/Source/Editor/GameProjectGeneration/Private/SProjectBrowser.cpp b/Engine/Source/Editor/GameProjectGeneration/Private/SProjectBrowser.cpp
--- a/Engine/Source/Editor/GameProjectGeneration/Private/SProjectBrowser.cpp
+++ b/Engine/Source/Editor/GameProjectGeneration/Private/SProjectBrowser.cpp
@@ -656,80 +656,42 @@ FReply SProjectBrowser::FindProjects()
 	const FText SamplesCategoryName = LOCTEXT("SamplesCategoryName", "Samples");
 
 	// Create a map of parent project folders to their category
-	TMap<FString, FText> ParentProjectFoldersToCategory;
-
-	// Add the default creation path, in case you've never created a project before or want to include this path anyway
-	ParentProjectFoldersToCategory.Add(GameProjectUtils::GetDefaultProjectCreationPath(), MyProjectsCategoryName);
-
-	// Add in every path that the user has ever created a project file. This is to catch new projects showing up in the user's project folders
-	const TArray<FString> &CreatedProjectPaths = GEditor->GetGameAgnosticSettings().CreatedProjectPaths;
-	for(int32 Idx = 0; Idx < CreatedProjectPaths.Num(); Idx++)
-	{
-		FString CreatedProjectPath = CreatedProjectPaths[Idx];
-		FPaths::NormalizeDirectoryName(CreatedProjectPath);
-		ParentProjectFoldersToCategory.Add(CreatedProjectPath, MyProjectsCategoryName);
-	}
+	TMap<FString, FText> ProjectFilesToCategory;
 
-	// Create a map of project folders to their category
-	TMap<FString, FText> ProjectFoldersToCategory;
+	// Find all the engine installations
+	TMap<FString, FString> EngineInstallations;
+	FDesktopPlatformModule::Get()->EnumerateEngineInstallations(EngineInstallations);
 
-	// Find all the subdirectories of all the parent folders
-	for(TMap<FString, FText>::TConstIterator Iter(ParentProjectFoldersToCategory); Iter; ++Iter)
+	// Add projects from every branch that we know about
+	FString CurrentEngineIdentifier = FDesktopPlatformModule::Get()->GetCurrentEngineIdentifier();
+	for(TMap<FString, FString>::TConstIterator Iter(EngineInstallations); Iter; ++Iter)
 	{
-		const FString &ParentProjectFolder = Iter.Key();
-
-		TArray<FString> ProjectFolders;
-		IFileManager::Get().FindFiles(ProjectFolders, *(ParentProjectFolder / TEXT("*")), false, true);
-
-		for(int32 Idx = 0; Idx < ProjectFolders.Num(); Idx++)
+		TArray<FString> ProjectFiles;
+		if(FDesktopPlatformModule::Get()->EnumerateProjectsKnownByEngine(Iter.Key(), false, ProjectFiles))
 		{
-			ProjectFoldersToCategory.Add(ParentProjectFolder / ProjectFolders[Idx], Iter.Value());
+			for(int Idx = 0; Idx < ProjectFiles.Num(); Idx++)
+			{
+				ProjectFilesToCategory.Add(ProjectFiles[Idx], MyProjectsCategoryName);
+			}
 		}
 	}
 
-	// Add all the launcher installed sample folders
-	TArray<FString> LauncherSampleDirectories;
-	FDesktopPlatformModule::Get()->EnumerateLauncherSampleInstallations(LauncherSampleDirectories);
-	for(int32 Idx = 0; Idx < LauncherSampleDirectories.Num(); Idx++)
+	// Add all the samples from the launcher
+	TArray<FString> LauncherSampleProjects;
+	FDesktopPlatformModule::Get()->EnumerateLauncherSampleProjects(LauncherSampleProjects);
+	for(int32 Idx = 0; Idx < LauncherSampleProjects.Num(); Idx++)
 	{
-		ProjectFoldersToCategory.Add(LauncherSampleDirectories[Idx], SamplesCategoryName);
-	}
-
-	// Create a map of all the project files to their category
-	TMap<FString, FText> ProjectFilesToCategory;
-
-	// Add all the folders that the user has opened recently
-	const TArray<FString> &RecentlyOpenedProjectFiles = GEditor->GetGameAgnosticSettings().RecentlyOpenedProjectFiles;
-	for (int32 Idx = 0; Idx < RecentlyOpenedProjectFiles.Num(); Idx++)
-	{
-		FString RecentlyOpenedProjectFile = RecentlyOpenedProjectFiles[Idx];
-		FPaths::NormalizeFilename(RecentlyOpenedProjectFile);
-		ProjectFilesToCategory.Add(RecentlyOpenedProjectFile, MyProjectsCategoryName);
-	}
-
-	// Scan the project folders for project files
-	FString ProjectWildcard = FString::Printf(TEXT("*.%s"), *IProjectManager::GetProjectFileExtension());
-	for(TMap<FString, FText>::TConstIterator Iter(ProjectFoldersToCategory); Iter; ++Iter)
-	{
-		const FString &ProjectFolder = Iter.Key();
-
-		TArray<FString> ProjectFiles;
-		IFileManager::Get().FindFiles(ProjectFiles, *(ProjectFolder / ProjectWildcard), true, false);
-
-		for(int32 Idx = 0; Idx < ProjectFiles.Num(); Idx++)
-		{
-			ProjectFilesToCategory.Add(ProjectFolder / ProjectFiles[Idx], Iter.Value());
-		}
+		ProjectFilesToCategory.Add(LauncherSampleProjects[Idx], SamplesCategoryName);
 	}
 
-	// Add all the discovered non-foreign project files
-	const TArray<FString> &NonForeignProjectFiles = FUProjectDictionary::GetDefault().GetProjectPaths();
-	for(int32 Idx = 0; Idx < NonForeignProjectFiles.Num(); Idx++)
+	// Add all the native project files we can find, and automatically filter them depending on their directory
+	const TArray<FString> &NativeProjectFiles = FUProjectDictionary::GetDefault().GetProjectPaths();
+	for(int32 Idx = 0; Idx < NativeProjectFiles.Num(); Idx++)
 	{
-		if(!NonForeignProjectFiles[Idx].Contains(TEXT("/Templates/")))
+		if(!NativeProjectFiles[Idx].Contains(TEXT("/Templates/")))
 		{
-			const FText &CategoryName = NonForeignProjectFiles[Idx].Contains(TEXT("/Samples/"))? SamplesCategoryName : MyProjectsCategoryName;
-			ProjectFilesToCategory.Add(NonForeignProjectFiles[Idx], CategoryName);
+			const FText &CategoryName = NativeProjectFiles[Idx].Contains(TEXT("/Samples/"))? SamplesCategoryName : MyProjectsCategoryName;
+			ProjectFilesToCategory.Add(NativeProjectFiles[Idx], CategoryName);
 		}
 	}
 
