Add a utility function for getting an engine's description from its identifier.

diff --git a/Engine/Source/Developer/DesktopPlatform/Private/DesktopPlatformBase.cpp b/Engine/Source/Developer/DesktopPlatform/Private/DesktopPlatformBase.cpp
--- a/Engine/Source/Developer/DesktopPlatform/Private/DesktopPlatformBase.cpp
+++ b/Engine/Source/Developer/DesktopPlatform/Private/DesktopPlatformBase.cpp
@@ -14,6 +14,36 @@ FDesktopPlatformBase::FDesktopPlatformBase()
 	LauncherInstallationTimestamp = FDateTime::MinValue();
 }
 
+FString FDesktopPlatformBase::GetEngineDescription(const FString& Identifier)
+{
+	// Official release versions just have a version number
+	if(IsStockEngineRelease(Identifier))
+	{
+		return Identifier;
+	}
+
+	// Otherwise get the path
+	FString RootDir;
+	if(!GetEngineRootDirFromIdentifier(Identifier, RootDir))
+	{
+		return FString();
+	}
+
+	// Convert it to a platform directory
+	FString PlatformRootDir = RootDir;
+	FPaths::MakePlatformFilename(PlatformRootDir);
+
+	// Perforce build
+	if (IsSourceDistribution(RootDir))
+	{
+		return FString::Printf(TEXT("Source build at %s"), *PlatformRootDir);
+	}
+	else
+	{
+		return FString::Printf(TEXT("Binary build at %s"), *PlatformRootDir);
+	}
+}
+
 FString FDesktopPlatformBase::GetCurrentEngineIdentifier()
 {
 	if(CurrentEngineIdentifier.Len() == 0 && !GetEngineIdentifierFromRootDir(FPlatformMisc::RootDir(), CurrentEngineIdentifier))
diff --git a/Engine/Source/Developer/DesktopPlatform/Private/DesktopPlatformBase.h b/Engine/Source/Developer/DesktopPlatform/Private/DesktopPlatformBase.h
--- a/Engine/Source/Developer/DesktopPlatform/Private/DesktopPlatformBase.h
+++ b/Engine/Source/Developer/DesktopPlatform/Private/DesktopPlatformBase.h
@@ -12,6 +12,7 @@ public:
 	FDesktopPlatformBase();
 
 	// IDesktopPlatform Implementation
+	virtual FString GetEngineDescription(const FString& Identifier) OVERRIDE;
 	virtual FString GetCurrentEngineIdentifier() OVERRIDE;
 
 	virtual void EnumerateLauncherEngineInstallations(TMap<FString, FString> &OutInstallations) OVERRIDE;
diff --git a/Engine/Source/Developer/DesktopPlatform/Public/IDesktopPlatform.h b/Engine/Source/Developer/DesktopPlatform/Public/IDesktopPlatform.h
--- a/Engine/Source/Developer/DesktopPlatform/Public/IDesktopPlatform.h
+++ b/Engine/Source/Developer/DesktopPlatform/Public/IDesktopPlatform.h
@@ -117,6 +117,13 @@ public:
 	 */
 	virtual bool OpenLauncher(bool Install, FString CommandLineParams ) = 0;
 
+	/**
+	* Returns a description for the engine with the given identifier.
+	*
+	* @return	Description for the engine identifier. Empty string if it's unknown.
+	*/
+	virtual FString GetEngineDescription(const FString& Identifier) = 0;
+
 	/**
 	* Gets the identifier for the currently executing engine installation
 	*
