Added PlatformInfo to DesktopPlatform and improved the editors Supported Platform UI

TTP# 337136 - SETTINGS: Target Platform settings polish
TTP# 337652 - EDITOR: Limit Project supported Android icons down to 1
TTP# 337650 - EDITOR: There is only 1 icon for Apple for Project Supported Platforms

DesktopPlatform now contains a static array of FPlatformInfo. This can be used to query UE4 about its available platforms, even when they're not available as a target platform.

FPlatformInfo contains the information required by the editor (such as a localized display name and icon), as well as whether a platform is a variation ("flavor") of another, and if so, whether the flavor affects the build output (eg, Win32 or Win64), or the cook output (eg, Android_XYZ). This lets the editor build up nested menus for the "Package Project" and "Cook Project" options, rather than just showing everything as a flat list.

diff --git a/Engine/Source/Developer/Android/AndroidTargetPlatform/AndroidTargetPlatform.Build.cs b/Engine/Source/Developer/Android/AndroidTargetPlatform/AndroidTargetPlatform.Build.cs
--- a/Engine/Source/Developer/Android/AndroidTargetPlatform/AndroidTargetPlatform.Build.cs
+++ b/Engine/Source/Developer/Android/AndroidTargetPlatform/AndroidTargetPlatform.Build.cs
@@ -14,6 +14,7 @@ public class AndroidTargetPlatform : ModuleRules
 				"Core",
 				"CoreUObject",
 				"TargetPlatform",
+				"DesktopPlatform",
 				"AndroidDeviceDetection",
 			}
 		);
diff --git a/Engine/Source/Developer/Android/AndroidTargetPlatform/Private/AndroidTargetPlatform.h b/Engine/Source/Developer/Android/AndroidTargetPlatform/Private/AndroidTargetPlatform.h
--- a/Engine/Source/Developer/Android/AndroidTargetPlatform/Private/AndroidTargetPlatform.h
+++ b/Engine/Source/Developer/Android/AndroidTargetPlatform/Private/AndroidTargetPlatform.h
@@ -96,8 +96,6 @@ public:
 
 	virtual ITargetDevicePtr GetDevice( const FTargetDeviceId& DeviceId ) OVERRIDE;
 
-	virtual FString GetIconPath( ETargetPlatformIcons::IconType IconType ) const OVERRIDE;
-
 	virtual bool IsRunningPlatform( ) const OVERRIDE;
 
 	virtual bool IsServerOnly( ) const OVERRIDE
diff --git a/Engine/Source/Developer/Android/AndroidTargetPlatform/Private/AndroidTargetPlatform.inl b/Engine/Source/Developer/Android/AndroidTargetPlatform/Private/AndroidTargetPlatform.inl
--- a/Engine/Source/Developer/Android/AndroidTargetPlatform/Private/AndroidTargetPlatform.inl
+++ b/Engine/Source/Developer/Android/AndroidTargetPlatform/Private/AndroidTargetPlatform.inl
@@ -75,22 +75,6 @@ inline ITargetDevicePtr FAndroidTargetPlatform::GetDevice( const FTargetDeviceId
 }
 
 
-inline FString FAndroidTargetPlatform::GetIconPath( ETargetPlatformIcons::IconType IconType ) const
-{
-	switch (IconType)
-	{
-	case ETargetPlatformIcons::Normal:
-		return FString::Printf(TEXT("Launcher/Android/Platform_%s_24x"), *PlatformName());
-
-	case ETargetPlatformIcons::Large:
-	case ETargetPlatformIcons::XLarge:
-		return FString(TEXT("Launcher/Android/Platform_Android_128x"));
-	}
-
-	return FString();
-}
-
-
 inline bool FAndroidTargetPlatform::IsRunningPlatform( ) const
 {
 	return false; // This platform never runs the target platform framework
diff --git a/Engine/Source/Developer/Android/AndroidTargetPlatform/Private/AndroidTargetPlatformPrivatePCH.h b/Engine/Source/Developer/Android/AndroidTargetPlatform/Private/AndroidTargetPlatformPrivatePCH.h
--- a/Engine/Source/Developer/Android/AndroidTargetPlatform/Private/AndroidTargetPlatformPrivatePCH.h
+++ b/Engine/Source/Developer/Android/AndroidTargetPlatform/Private/AndroidTargetPlatformPrivatePCH.h
@@ -21,6 +21,7 @@
 #endif
 
 #include "TargetPlatform.h"
+#include "TargetPlatformBase.h"
 #include "AndroidDeviceDetection.h"
 
 
diff --git a/Engine/Source/Developer/Android/Android_ATCTargetPlatform/Android_ATCTargetPlatform.Build.cs b/Engine/Source/Developer/Android/Android_ATCTargetPlatform/Android_ATCTargetPlatform.Build.cs
--- a/Engine/Source/Developer/Android/Android_ATCTargetPlatform/Android_ATCTargetPlatform.Build.cs
+++ b/Engine/Source/Developer/Android/Android_ATCTargetPlatform/Android_ATCTargetPlatform.Build.cs
@@ -13,6 +13,7 @@ public class Android_ATCTargetPlatform : ModuleRules
 			{
 				"Core",
 				"TargetPlatform",
+				"DesktopPlatform",
 				"AndroidDeviceDetection",
 			}
 		);
diff --git a/Engine/Source/Developer/Android/Android_ATCTargetPlatform/Private/Android_ATCTargetPlatformPrivatePCH.h b/Engine/Source/Developer/Android/Android_ATCTargetPlatform/Private/Android_ATCTargetPlatformPrivatePCH.h
--- a/Engine/Source/Developer/Android/Android_ATCTargetPlatform/Private/Android_ATCTargetPlatformPrivatePCH.h
+++ b/Engine/Source/Developer/Android/Android_ATCTargetPlatform/Private/Android_ATCTargetPlatformPrivatePCH.h
@@ -21,6 +21,7 @@
 #endif
 
 #include "TargetPlatform.h"
+#include "TargetPlatformBase.h"
 #include "AndroidDeviceDetection.h"
 
 
diff --git a/Engine/Source/Developer/Android/Android_DXTTargetPlatform/Android_DXTTargetPlatform.Build.cs b/Engine/Source/Developer/Android/Android_DXTTargetPlatform/Android_DXTTargetPlatform.Build.cs
--- a/Engine/Source/Developer/Android/Android_DXTTargetPlatform/Android_DXTTargetPlatform.Build.cs
+++ b/Engine/Source/Developer/Android/Android_DXTTargetPlatform/Android_DXTTargetPlatform.Build.cs
@@ -13,6 +13,7 @@ public class Android_DXTTargetPlatform : ModuleRules
 			{
 				"Core",
 				"TargetPlatform",
+				"DesktopPlatform",
 				"AndroidDeviceDetection",
 			}
 		);
diff --git a/Engine/Source/Developer/Android/Android_DXTTargetPlatform/Private/Android_DXTTargetPlatformPrivatePCH.h b/Engine/Source/Developer/Android/Android_DXTTargetPlatform/Private/Android_DXTTargetPlatformPrivatePCH.h
--- a/Engine/Source/Developer/Android/Android_DXTTargetPlatform/Private/Android_DXTTargetPlatformPrivatePCH.h
+++ b/Engine/Source/Developer/Android/Android_DXTTargetPlatform/Private/Android_DXTTargetPlatformPrivatePCH.h
@@ -21,6 +21,7 @@
 #endif
 
 #include "TargetPlatform.h"
+#include "TargetPlatformBase.h"
 #include "AndroidDeviceDetection.h"
 
 
diff --git a/Engine/Source/Developer/Android/Android_ETC1TargetPlatform/Android_ETC1TargetPlatform.Build.cs b/Engine/Source/Developer/Android/Android_ETC1TargetPlatform/Android_ETC1TargetPlatform.Build.cs
--- a/Engine/Source/Developer/Android/Android_ETC1TargetPlatform/Android_ETC1TargetPlatform.Build.cs
+++ b/Engine/Source/Developer/Android/Android_ETC1TargetPlatform/Android_ETC1TargetPlatform.Build.cs
@@ -20,6 +20,7 @@ public class Android_ETC1TargetPlatform : ModuleRules
 			{
 				"Core",
 				"TargetPlatform",
+				"DesktopPlatform",
 				"AndroidDeviceDetection",
 			}
 		);
diff --git a/Engine/Source/Developer/Android/Android_ETC1TargetPlatform/Private/Android_ETC1TargetPlatformPrivatePCH.h b/Engine/Source/Developer/Android/Android_ETC1TargetPlatform/Private/Android_ETC1TargetPlatformPrivatePCH.h
--- a/Engine/Source/Developer/Android/Android_ETC1TargetPlatform/Private/Android_ETC1TargetPlatformPrivatePCH.h
+++ b/Engine/Source/Developer/Android/Android_ETC1TargetPlatform/Private/Android_ETC1TargetPlatformPrivatePCH.h
@@ -21,6 +21,7 @@
 #endif
 
 #include "TargetPlatform.h"
+#include "TargetPlatformBase.h"
 #include "AndroidDeviceDetection.h"
 
 
diff --git a/Engine/Source/Developer/Android/Android_ETC2TargetPlatform/Android_ETC2TargetPlatform.Build.cs b/Engine/Source/Developer/Android/Android_ETC2TargetPlatform/Android_ETC2TargetPlatform.Build.cs
--- a/Engine/Source/Developer/Android/Android_ETC2TargetPlatform/Android_ETC2TargetPlatform.Build.cs
+++ b/Engine/Source/Developer/Android/Android_ETC2TargetPlatform/Android_ETC2TargetPlatform.Build.cs
@@ -13,6 +13,7 @@ public class Android_ETC2TargetPlatform : ModuleRules
 			{
 				"Core",
 				"TargetPlatform",
+				"DesktopPlatform",
 				"AndroidDeviceDetection",
 			}
 		);
diff --git a/Engine/Source/Developer/Android/Android_ETC2TargetPlatform/Private/Android_ETC2TargetPlatformPrivatePCH.h b/Engine/Source/Developer/Android/Android_ETC2TargetPlatform/Private/Android_ETC2TargetPlatformPrivatePCH.h
--- a/Engine/Source/Developer/Android/Android_ETC2TargetPlatform/Private/Android_ETC2TargetPlatformPrivatePCH.h
+++ b/Engine/Source/Developer/Android/Android_ETC2TargetPlatform/Private/Android_ETC2TargetPlatformPrivatePCH.h
@@ -21,6 +21,7 @@
 #endif
 
 #include "TargetPlatform.h"
+#include "TargetPlatformBase.h"
 #include "AndroidDeviceDetection.h"
 
 
diff --git a/Engine/Source/Developer/Android/Android_PVRTCTargetPlatform/Android_PVRTCTargetPlatform.Build.cs b/Engine/Source/Developer/Android/Android_PVRTCTargetPlatform/Android_PVRTCTargetPlatform.Build.cs
--- a/Engine/Source/Developer/Android/Android_PVRTCTargetPlatform/Android_PVRTCTargetPlatform.Build.cs
+++ b/Engine/Source/Developer/Android/Android_PVRTCTargetPlatform/Android_PVRTCTargetPlatform.Build.cs
@@ -13,6 +13,7 @@ public class Android_PVRTCTargetPlatform : ModuleRules
 			{
 				"Core",
 				"TargetPlatform",
+				"DesktopPlatform",
 				"AndroidDeviceDetection",
 			}
 		);
diff --git a/Engine/Source/Developer/Android/Android_PVRTCTargetPlatform/Private/Android_PVRTCTargetPlatformPrivatePCH.h b/Engine/Source/Developer/Android/Android_PVRTCTargetPlatform/Private/Android_PVRTCTargetPlatformPrivatePCH.h
--- a/Engine/Source/Developer/Android/Android_PVRTCTargetPlatform/Private/Android_PVRTCTargetPlatformPrivatePCH.h
+++ b/Engine/Source/Developer/Android/Android_PVRTCTargetPlatform/Private/Android_PVRTCTargetPlatformPrivatePCH.h
@@ -21,6 +21,7 @@
 #endif
 
 #include "TargetPlatform.h"
+#include "TargetPlatformBase.h"
 #include "AndroidDeviceDetection.h"
 
 
diff --git a/Engine/Source/Developer/DesktopPlatform/Private/PlatformInfo.cpp b/Engine/Source/Developer/DesktopPlatform/Private/PlatformInfo.cpp
new file mode 100644
--- /dev/null
+++ b/Engine/Source/Developer/DesktopPlatform/Private/PlatformInfo.cpp
@@ -0,0 +1,159 @@
+// Copyright 1998-2014 Epic Games, Inc. All Rights Reserved.
+
+#include "DesktopPlatformPrivatePCH.h"
+#include "PlatformInfo.h"
+
+#define LOCTEXT_NAMESPACE "PlatformInfo"
+
+namespace PlatformInfo
+{
+
+namespace
+{
+
+FPlatformInfo BuildPlatformInfo(const FName& InPlatformInfoName, const FName& InTargetPlatformName, const FText& InDisplayName, const EPlatformType InPlatformType, const EPlatformFlags::Flags InPlatformFlags, const FPlatformIconPaths& InIconPaths, const FString& InUATCommandLine)
+{
+	FPlatformInfo PlatformInfo;
+
+	PlatformInfo.PlatformInfoName = InPlatformInfoName;
+	PlatformInfo.TargetPlatformName = InTargetPlatformName;
+
+	// See if this name also contains a flavor
+	const FString InPlatformInfoNameString = InPlatformInfoName.ToString();
+	{
+		int32 UnderscoreLoc;
+		if(InPlatformInfoNameString.FindChar(TEXT('_'), UnderscoreLoc))
+		{
+			PlatformInfo.VanillaPlatformName = *InPlatformInfoNameString.Mid(0, UnderscoreLoc);
+			PlatformInfo.PlatformFlavor = *InPlatformInfoNameString.Mid(UnderscoreLoc + 1);
+		}
+		else
+		{
+			PlatformInfo.VanillaPlatformName = InPlatformInfoName;
+		}
+	}
+
+	PlatformInfo.DisplayName = InDisplayName;
+	PlatformInfo.PlatformType = InPlatformType;
+	PlatformInfo.PlatformFlags = InPlatformFlags;
+	PlatformInfo.IconPaths = InIconPaths;
+	PlatformInfo.UATCommandLine = InUATCommandLine;
+
+	// Generate the icon style names for FEditorStyle
+	PlatformInfo.IconPaths.NormalStyleName = *FString::Printf(TEXT("Launcher.Platform_%s"), *InPlatformInfoNameString);
+	PlatformInfo.IconPaths.LargeStyleName  = *FString::Printf(TEXT("Launcher.Platform_%s.Large"), *InPlatformInfoNameString);
+	PlatformInfo.IconPaths.XLargeStyleName = *FString::Printf(TEXT("Launcher.Platform_%s.XLarge"), *InPlatformInfoNameString);
+
+	return PlatformInfo;
+}
+
+static const FPlatformInfo PlatformInfoArray[] = {
+	//				  PlatformInfoName					TargetPlatformName			DisplayName													PlatformType				PlatformFlags					IconPaths																																UATCommandLine
+	BuildPlatformInfo(TEXT("WindowsNoEditor"),			TEXT("WindowsNoEditor"),	LOCTEXT("WindowsNoEditor", "Windows"),						EPlatformType::Game,		EPlatformFlags::None,			FPlatformIconPaths(TEXT("Launcher/Windows/Platform_WindowsNoEditor_24x"), TEXT("Launcher/Windows/Platform_WindowsNoEditor_128x")),		TEXT("")),
+	BuildPlatformInfo(TEXT("WindowsNoEditor_Win32"),	TEXT("WindowsNoEditor"),	LOCTEXT("WindowsNoEditor_Win32", "Windows (32-bit)"),		EPlatformType::Game,		EPlatformFlags::BuildFlavor,	FPlatformIconPaths(TEXT("Launcher/Windows/Platform_WindowsNoEditor_24x"), TEXT("Launcher/Windows/Platform_WindowsNoEditor_128x")),		TEXT("-targetplatform=Win32")),
+	BuildPlatformInfo(TEXT("WindowsNoEditor_Win64"),	TEXT("WindowsNoEditor"),	LOCTEXT("WindowsNoEditor_Win64", "Windows (64-bit)"),		EPlatformType::Game,		EPlatformFlags::BuildFlavor,	FPlatformIconPaths(TEXT("Launcher/Windows/Platform_WindowsNoEditor_24x"), TEXT("Launcher/Windows/Platform_WindowsNoEditor_128x")),		TEXT("-targetplatform=Win64")),
+	BuildPlatformInfo(TEXT("Windows"),					TEXT("Windows"),			LOCTEXT("WindowsEditor", "Windows (Editor)"),				EPlatformType::Editor,		EPlatformFlags::None,			FPlatformIconPaths(TEXT("Launcher/Windows/Platform_Windows_24x"), TEXT("Launcher/Windows/Platform_Windows_128x")),						TEXT("")),
+	BuildPlatformInfo(TEXT("WindowsClient"),			TEXT("WindowsClient"),		LOCTEXT("WindowsClient", "Windows (Client-only)"),			EPlatformType::Client,		EPlatformFlags::None,			FPlatformIconPaths(TEXT("Launcher/Windows/Platform_Windows_24x"), TEXT("Launcher/Windows/Platform_Windows_128x")),						TEXT("")),
+	BuildPlatformInfo(TEXT("WindowsServer"),			TEXT("WindowsServer"),		LOCTEXT("WindowsServer", "Windows (Dedicated Server)"),		EPlatformType::Server,		EPlatformFlags::None,			FPlatformIconPaths(TEXT("Launcher/Windows/Platform_WindowsServer_24x"), TEXT("Launcher/Windows/Platform_WindowsServer_128x")),			TEXT("")),
+
+	BuildPlatformInfo(TEXT("MacNoEditor"),				TEXT("MacNoEditor"),		LOCTEXT("MacNoEditor", "Mac"),								EPlatformType::Game,		EPlatformFlags::None,			FPlatformIconPaths(TEXT("Launcher/Mac/Platform_Mac_24x"), TEXT("Launcher/Mac/Platform_Mac_128x")),										TEXT("")),
+	BuildPlatformInfo(TEXT("Mac"),						TEXT("Mac"),				LOCTEXT("MacEditor", "Mac (Editor)"),						EPlatformType::Editor,		EPlatformFlags::None,			FPlatformIconPaths(TEXT("Launcher/Mac/Platform_Mac_24x"), TEXT("Launcher/Mac/Platform_Mac_128x")),										TEXT("")),
+	BuildPlatformInfo(TEXT("MacClient"),				TEXT("MacClient"),			LOCTEXT("MacClient", "Mac (Client-only)"),					EPlatformType::Client,		EPlatformFlags::None,			FPlatformIconPaths(TEXT("Launcher/Mac/Platform_Mac_24x"), TEXT("Launcher/Mac/Platform_Mac_128x")),										TEXT("")),
+	BuildPlatformInfo(TEXT("MacServer"),				TEXT("MacServer"),			LOCTEXT("MacServer", "Mac (Dedicated Server)"),				EPlatformType::Server,		EPlatformFlags::None,			FPlatformIconPaths(TEXT("Launcher/Mac/Platform_Mac_24x"), TEXT("Launcher/Mac/Platform_Mac_128x")),										TEXT("")),
+
+	BuildPlatformInfo(TEXT("LinuxNoEditor"),			TEXT("LinuxNoEditor"),		LOCTEXT("LinuxNoEditor", "Linux"),							EPlatformType::Game,		EPlatformFlags::None,			FPlatformIconPaths(TEXT("Launcher/Linux/Platform_Linux_24x"), TEXT("Launcher/Linux/Platform_Linux_128x")),								TEXT("")),
+	BuildPlatformInfo(TEXT("Linux"),					TEXT("Linux"),				LOCTEXT("LinuxEditor", "Linux (Editor)"),					EPlatformType::Editor,		EPlatformFlags::None,			FPlatformIconPaths(TEXT("Launcher/Linux/Platform_Linux_24x"), TEXT("Launcher/Linux/Platform_Linux_128x")),								TEXT("")),
+	BuildPlatformInfo(TEXT("LinuxClient"),				TEXT("LinuxClient"),		LOCTEXT("LinuxClient", "Linux (Client-only)"),				EPlatformType::Client,		EPlatformFlags::None,			FPlatformIconPaths(TEXT("Launcher/Linux/Platform_Linux_24x"), TEXT("Launcher/Linux/Platform_Linux_128x")),								TEXT("")),
+	BuildPlatformInfo(TEXT("LinuxServer"),				TEXT("LinuxServer"),		LOCTEXT("LinuxServer", "Linux (Dedicated Server)"),			EPlatformType::Server,		EPlatformFlags::None,			FPlatformIconPaths(TEXT("Launcher/Linux/Platform_Linux_24x"), TEXT("Launcher/Linux/Platform_Linux_128x")),								TEXT("")),
+	
+	BuildPlatformInfo(TEXT("Android"),					TEXT("Android"),			LOCTEXT("Android", "Android"),								EPlatformType::Game,		EPlatformFlags::None,			FPlatformIconPaths(TEXT("Launcher/Android/Platform_Android_24x"), TEXT("Launcher/Android/Platform_Android_128x")),						TEXT("")),
+	BuildPlatformInfo(TEXT("Android_All"),				TEXT("Android"),			LOCTEXT("Android_All", "Android (All)"),					EPlatformType::Game,		EPlatformFlags::CookFlavor,		FPlatformIconPaths(TEXT("Launcher/Android/Platform_Android_24x"), TEXT("Launcher/Android/Platform_Android_128x")),						TEXT("")),
+	BuildPlatformInfo(TEXT("Android_ATC"),				TEXT("Android_ATC"),		LOCTEXT("Android_ATC", "Android (ATC)"),					EPlatformType::Game,		EPlatformFlags::CookFlavor,		FPlatformIconPaths(TEXT("Launcher/Android/Platform_Android_ATC_24x"), TEXT("Launcher/Android/Platform_Android_128x")),					TEXT("-cookflavor=ATC")),
+	BuildPlatformInfo(TEXT("Android_DXT"),				TEXT("Android_DXT"),		LOCTEXT("Android_DXT", "Android (DXT)"),					EPlatformType::Game,		EPlatformFlags::CookFlavor,		FPlatformIconPaths(TEXT("Launcher/Android/Platform_Android_DXT_24x"), TEXT("Launcher/Android/Platform_Android_128x")),					TEXT("-cookflavor=DXT")),
+	BuildPlatformInfo(TEXT("Android_ETC1"),				TEXT("Android_ETC1"),		LOCTEXT("Android_ETC1", "Android (ETC1)"),					EPlatformType::Game,		EPlatformFlags::CookFlavor,		FPlatformIconPaths(TEXT("Launcher/Android/Platform_Android_ETC1_24x"), TEXT("Launcher/Android/Platform_Android_128x")),					TEXT("-cookflavor=ETC1")),
+	BuildPlatformInfo(TEXT("Android_ETC2"),				TEXT("Android_ETC2"),		LOCTEXT("Android_ETC2", "Android (ETC2)"),					EPlatformType::Game,		EPlatformFlags::CookFlavor,		FPlatformIconPaths(TEXT("Launcher/Android/Platform_Android_ETC2_24x"), TEXT("Launcher/Android/Platform_Android_128x")),					TEXT("-cookflavor=ETC2")),
+	BuildPlatformInfo(TEXT("Android_PVRTC"),			TEXT("Android_PVRTC"),		LOCTEXT("Android_PVRTC", "Android (PVRTC)"),				EPlatformType::Game,		EPlatformFlags::CookFlavor,		FPlatformIconPaths(TEXT("Launcher/Android/Platform_Android_PVRTC_24x"), TEXT("Launcher/Android/Platform_Android_128x")),				TEXT("-cookflavor=PVRTC")),
+
+	BuildPlatformInfo(TEXT("IOS"),						TEXT("IOS"),				LOCTEXT("IOS", "iOS"),										EPlatformType::Game,		EPlatformFlags::None,			FPlatformIconPaths(TEXT("Launcher/iOS/Platform_iOS_24x"), TEXT("Launcher/iOS/Platform_iOS_128x")),										TEXT("")),
+
+	BuildPlatformInfo(TEXT("PS4"),						TEXT("PS4"),				LOCTEXT("PS4", "PlayStation 4"),							EPlatformType::Game,		EPlatformFlags::None,			FPlatformIconPaths(TEXT("Launcher/PS4/Platform_PS4_24x"), TEXT("Launcher/PS4/Platform_PS4_128x")),										TEXT("")),
+
+	BuildPlatformInfo(TEXT("XboxOne"),					TEXT("XboxOne"),			LOCTEXT("XboxOne", "Xbox One"),								EPlatformType::Game,		EPlatformFlags::None,			FPlatformIconPaths(TEXT("Launcher/XboxOne/Platform_XboxOne_24x"), TEXT("Launcher/XboxOne/Platform_XboxOne_128x")),						TEXT("")),
+
+	BuildPlatformInfo(TEXT("HTML5"),					TEXT("HTML5"),				LOCTEXT("HTML5", "HTML5"),									EPlatformType::Game,		EPlatformFlags::None,			FPlatformIconPaths(TEXT("Launcher/HTML5/Platform_HTML5_24x"), TEXT("Launcher/HTML5/Platform_HTML5_128x")),								TEXT("")),
+
+	//BuildPlatformInfo(TEXT("WinRT"),					TEXT("WinRT"),				LOCTEXT("WinRT", "Windows RT"),								EPlatformType::Game,		EPlatformFlags::None,			FPlatformIconPaths(TEXT("Launcher/Windows/Platform_WindowsNoEditor_24x"), TEXT("Launcher/Windows/Platform_WindowsNoEditor_128x")),		TEXT("")),
+};
+
+} // anonymous namespace
+
+const FPlatformInfo* FindPlatformInfo(const FName& InPlatformName)
+{
+	for(const FPlatformInfo& PlatformInfo : PlatformInfoArray)
+	{
+		if(PlatformInfo.PlatformInfoName == InPlatformName)
+		{
+			return &PlatformInfo;
+		}
+	}
+	return nullptr;
+}
+
+const FPlatformInfo* FindVanillaPlatformInfo(const FName& InPlatformName)
+{
+	const FPlatformInfo* const FoundInfo = FindPlatformInfo(InPlatformName);
+	return (FoundInfo) ? (FoundInfo->IsVanilla()) ? FoundInfo : FindPlatformInfo(FoundInfo->VanillaPlatformName) : nullptr;
+}
+
+const FPlatformInfo* GetPlatformInfoArray(int32& OutNumPlatforms)
+{
+	OutNumPlatforms = ARRAY_COUNT(PlatformInfoArray);
+	return PlatformInfoArray;
+}
+
+FPlatformEnumerator EnumeratePlatformInfoArray()
+{
+	return FPlatformEnumerator(PlatformInfoArray, ARRAY_COUNT(PlatformInfoArray));
+}
+
+TArray<FVanillaPlatformEntry> BuildPlatformHierarchy(const EPlatformFilter InFilter)
+{
+	TArray<FVanillaPlatformEntry> VanillaPlatforms;
+
+	// Build up a tree from the platforms we support (vanilla outers, with a list of flavors)
+	// PlatformInfoArray should be ordered in such a way that the vanilla platforms always appear before their flavors
+	for(const FPlatformInfo& PlatformInfo : PlatformInfoArray)
+	{
+		if(PlatformInfo.IsVanilla())
+		{
+			VanillaPlatforms.Add(FVanillaPlatformEntry(&PlatformInfo));
+		}
+		else
+		{
+			const bool bHasBuildFlavor = !!(PlatformInfo.PlatformFlags & EPlatformFlags::BuildFlavor);
+			const bool bHasCookFlavor  = !!(PlatformInfo.PlatformFlags & EPlatformFlags::CookFlavor);
+			
+			const bool bValidFlavor = 
+				InFilter == EPlatformFilter::All || 
+				(InFilter == EPlatformFilter::BuildFlavor && bHasBuildFlavor) || 
+				(InFilter == EPlatformFilter::CookFlavor && bHasCookFlavor);
+
+			if(bValidFlavor)
+			{
+				const FName VanillaPlatformName = PlatformInfo.VanillaPlatformName;
+				FVanillaPlatformEntry* const VanillaEntry = VanillaPlatforms.FindByPredicate([VanillaPlatformName](const FVanillaPlatformEntry& Item) -> bool
+				{
+					return Item.PlatformInfo->PlatformInfoName == VanillaPlatformName;
+				});
+				check(VanillaEntry);
+				VanillaEntry->PlatformFlavors.Add(&PlatformInfo);
+			}
+		}
+	}
+
+	return VanillaPlatforms;
+}
+
+} // namespace PlatformFamily
+
+#undef LOCTEXT_NAMESPACE
diff --git a/Engine/Source/Developer/DesktopPlatform/Public/PlatformInfo.h b/Engine/Source/Developer/DesktopPlatform/Public/PlatformInfo.h
new file mode 100644
--- /dev/null
+++ b/Engine/Source/Developer/DesktopPlatform/Public/PlatformInfo.h
@@ -0,0 +1,285 @@
+// Copyright 1998-2014 Epic Games, Inc. All Rights Reserved.
+
+#pragma once
+
+namespace PlatformInfo
+{
+	/** The target type of the given platform */
+	enum class EPlatformType : uint8
+	{
+		/** The platform targets cooked monolithic game executables */
+		Game,
+
+		/** The platform targets uncooked modular editor executables and DLLs */
+		Editor,
+
+		/** The platform targets cooked monolithic game client executables (but no server code) */
+		Client,
+
+		/** The platform targets cooked monolithic game server executables (but no client code) */
+		Server,
+	};
+
+	/** Available icon sizes (see FPlatformIconPaths) */
+	enum class EPlatformIconSize : uint8
+	{
+		/** Normal sized icon (24x24) */
+		Normal,
+
+		/** Large sized icon (64x64) */
+		Large,
+
+		/** Extra large sized icon (128x128) */
+		XLarge,
+	};
+
+	/** Flavor types used when filtering the platforms based upon their flags */
+	enum class EPlatformFilter : uint8
+	{
+		/** Include all platform types */
+		All,
+
+		/** Include only build flavors */
+		BuildFlavor,
+
+		/** Include only cook flavors */
+		CookFlavor,
+	};
+
+	/** Flags describing platform variants */
+	namespace EPlatformFlags
+	{
+		typedef uint8 Flags;
+		enum Flag
+		{
+			/** Nothing of interest */
+			None = 0,
+
+			/** The flavor generates different output when building (eg, 32 or 64-bit) */
+			BuildFlavor = 1<<0,
+
+			/** The flavor generates different output when cooking (eg, ATC or PVRTC texture format) */
+			CookFlavor = 1<<1,
+		};
+	}
+
+	/** Information about where to find the platform icons (for use by FEditorStyle) */
+	struct FPlatformIconPaths
+	{
+		FPlatformIconPaths()
+		{
+		}
+
+		FPlatformIconPaths(const FString& InIconPath)
+			: NormalPath(InIconPath)
+			, LargePath(InIconPath)
+			, XLargePath(InIconPath)
+		{
+		}
+
+		FPlatformIconPaths(const FString& InNormalPath, const FString& InLargePath)
+			: NormalPath(InNormalPath)
+			, LargePath(InLargePath)
+			, XLargePath(InLargePath)
+		{
+		}
+
+		FPlatformIconPaths(const FString& InNormalPath, const FString& InLargePath, const FString& InXLargePath)
+			: NormalPath(InNormalPath)
+			, LargePath(InLargePath)
+			, XLargePath(InXLargePath)
+		{
+		}
+
+		FName NormalStyleName;
+		FString NormalPath;
+
+		FName LargeStyleName;
+		FString LargePath;
+
+		FName XLargeStyleName;
+		FString XLargePath;
+	};
+
+	/** Information about a given platform */
+	struct FPlatformInfo
+	{
+		/** Name used to identify this platform, eg "Android_ATC" */
+		FName PlatformInfoName;
+
+		/** Name used to find the corresponding ITargetPlatform for this platform (also used by UAT) */
+		FName TargetPlatformName;
+
+		/** Vanilla name for this platform, eg "Android" for "Android_ATC" */
+		FName VanillaPlatformName;
+
+		/** Platform flavor, eg "ATC" for "Android_ATC" */
+		FName PlatformFlavor;
+
+		/** The friendly (and localized) display name of this platform */
+		FText DisplayName;
+
+		/** Type of this platform */
+		EPlatformType PlatformType;
+
+		/** Flags for this platform */
+		EPlatformFlags::Flags PlatformFlags;
+
+		/** Information about where to find the platform icons (for use by FEditorStyle) */
+		FPlatformIconPaths IconPaths;
+
+		/** Additional argument string data to append to UAT commands relating to this platform */
+		FString UATCommandLine;
+
+		/** Returns true if this platform is vanilla */
+		FORCEINLINE bool IsVanilla() const
+		{
+			return PlatformFlavor.IsNone();
+		}
+
+		/** Returns true if this platform is a flavor */
+		FORCEINLINE bool IsFlavor() const
+		{
+			return !PlatformFlavor.IsNone();
+		}
+
+		/** Get the icon name (for FEditorStyle) used by the given icon type for this platform */
+		FName GetIconStyleName(const EPlatformIconSize InIconSize) const
+		{
+			switch(InIconSize)
+			{
+			case EPlatformIconSize::Normal:
+				return IconPaths.NormalStyleName;
+			case EPlatformIconSize::Large:
+				return IconPaths.LargeStyleName;
+			case EPlatformIconSize::XLarge:
+				return IconPaths.XLargeStyleName;
+			default:
+				break;
+			}
+			return NAME_None;
+		}
+
+		/** Get the path to the icon on disk (for FEditorStyle) for the given icon type for this platform */
+		const FString& GetIconPath(const EPlatformIconSize InIconSize) const
+		{
+			switch(InIconSize)
+			{
+			case EPlatformIconSize::Normal:
+				return IconPaths.NormalPath;
+			case EPlatformIconSize::Large:
+				return IconPaths.LargePath;
+			case EPlatformIconSize::XLarge:
+				return IconPaths.XLargePath;
+			default:
+				break;
+			}
+			static const FString EmptyString = TEXT("");
+			return EmptyString;
+		}
+	};
+
+	/** Vanilla platform that may contain a set of flavors */
+	struct FVanillaPlatformEntry
+	{
+		FVanillaPlatformEntry()
+			: PlatformInfo(nullptr)
+		{
+		}
+
+		FVanillaPlatformEntry(const FPlatformInfo* const InPlatformInfo)
+			: PlatformInfo(InPlatformInfo)
+		{
+		}
+
+		/** Information for this platform */
+		const FPlatformInfo* PlatformInfo;
+
+		/** Any flavors for this platform */
+		TArray<const FPlatformInfo*> PlatformFlavors;
+	};
+
+	/** Simple wrapper class to allow range-based-for enumeration from a call to EnumeratePlatforms */
+	class FPlatformEnumerator
+	{
+	public:
+		FPlatformEnumerator(const FPlatformInfo* InPlatformsArray, const int32 InNumPlatforms)
+			: PlatformsArray(InPlatformsArray)
+			, CurrentPlatform(InPlatformsArray)
+			, NumPlatforms(InNumPlatforms)
+		{
+		}
+
+		FORCEINLINE const FPlatformInfo* begin() const
+		{
+			return PlatformsArray;
+		}
+
+		FORCEINLINE const FPlatformInfo* end() const
+		{
+			return PlatformsArray + NumPlatforms;
+		}
+
+		FORCEINLINE const FPlatformInfo* operator->() const
+		{
+			return CurrentPlatform;
+		}
+
+		FORCEINLINE const FPlatformInfo& operator*() const
+		{
+			return *CurrentPlatform;
+		}
+
+		FORCEINLINE FPlatformEnumerator& operator++()
+		{
+			++CurrentPlatform;
+			return *this;
+		}
+
+		FORCEINLINE FPlatformEnumerator operator++(int)
+		{
+			FPlatformEnumerator Copy(*this);
+			++CurrentPlatform;
+			return Copy;
+		}
+
+	private:
+		const FPlatformInfo* PlatformsArray;
+		const FPlatformInfo* CurrentPlatform;
+		int32 NumPlatforms;
+	};
+
+	/**
+	 * Try and find the information for the given platform
+	 * @param InPlatformName - The name of the platform to find
+	 * @return The platform info if the platform was found, null otherwise
+	 */
+	DESKTOPPLATFORM_API const FPlatformInfo* FindPlatformInfo(const FName& InPlatformName);
+
+	/**
+	 * Try and find the vanilla information for the given platform
+	 * @param InPlatformName - The name of the platform to find (can be a flavor, but you'll still get back the vanilla platform)
+	 * @return The platform info if the platform was found, null otherwise
+	 */
+	DESKTOPPLATFORM_API const FPlatformInfo* FindVanillaPlatformInfo(const FName& InPlatformName);
+
+	/**
+	 * Get an array of all the platforms we know about
+	 * @param OutNumPlatforms - The number of platforms in the array
+	 * @return The pointer to the start of the platform array
+	 */
+	DESKTOPPLATFORM_API const FPlatformInfo* GetPlatformInfoArray(int32& OutNumPlatforms);
+
+	/**
+	 * Convenience function to enumerate all the platforms we know about (compatible with range-based-for)
+	 * @return An enumerator for the platforms (see FPlatformEnumerator)
+	 */
+	DESKTOPPLATFORM_API FPlatformEnumerator EnumeratePlatformInfoArray();
+
+	/**
+	 * Build a hierarchy mapping vanilla platforms to their flavors
+	 * @param InFilter - Flags to control which kinds of flavors you want to include
+	 * @return An array of vanilla platforms, potentially containing flavors
+	 */
+	DESKTOPPLATFORM_API TArray<FVanillaPlatformEntry> BuildPlatformHierarchy(const EPlatformFilter InFilter);
+}
diff --git a/Engine/Source/Developer/DeviceManager/DeviceManager.Build.cs b/Engine/Source/Developer/DeviceManager/DeviceManager.Build.cs
--- a/Engine/Source/Developer/DeviceManager/DeviceManager.Build.cs
+++ b/Engine/Source/Developer/DeviceManager/DeviceManager.Build.cs
@@ -20,6 +20,7 @@ public class DeviceManager : ModuleRules
 				"Slate",
 				"SlateCore",
 				"TargetPlatform",
+				"DesktopPlatform",
 			}
 		);
 
diff --git a/Engine/Source/Developer/DeviceManager/Private/DeviceManagerPrivatePCH.h b/Engine/Source/Developer/DeviceManager/Private/DeviceManagerPrivatePCH.h
--- a/Engine/Source/Developer/DeviceManager/Private/DeviceManagerPrivatePCH.h
+++ b/Engine/Source/Developer/DeviceManager/Private/DeviceManagerPrivatePCH.h
@@ -14,6 +14,7 @@
 
 #include "ModuleManager.h"
 #include "TargetPlatform.h"
+#include "PlatformInfo.h"
 
 
 /* Private includes
diff --git a/Engine/Source/Developer/DeviceManager/Private/Widgets/Browser/SDeviceBrowserDeviceAdder.cpp b/Engine/Source/Developer/DeviceManager/Private/Widgets/Browser/SDeviceBrowserDeviceAdder.cpp
--- a/Engine/Source/Developer/DeviceManager/Private/Widgets/Browser/SDeviceBrowserDeviceAdder.cpp
+++ b/Engine/Source/Developer/DeviceManager/Private/Widgets/Browser/SDeviceBrowserDeviceAdder.cpp
@@ -303,6 +303,9 @@ FString SDeviceBrowserDeviceAdder::HandlePlatformComboBoxContentText( ) const
 
 TSharedRef<SWidget> SDeviceBrowserDeviceAdder::HandlePlatformComboBoxGenerateWidget( TSharedPtr<FString> Item )
 {
+	const PlatformInfo::FPlatformInfo* const PlatformInfo = PlatformInfo::FindPlatformInfo(**Item);
+	check(PlatformInfo);
+
 	return
 		SNew(SHorizontalBox)
 
@@ -311,7 +314,7 @@ TSharedRef<SWidget> SDeviceBrowserDeviceAdder::HandlePlatformComboBoxGenerateWid
 			.HAlign(HAlign_Left)
 			[
 				SNew(SImage)
-					.Image(FEditorStyle::GetBrush(*FString::Printf(TEXT("Launcher.Platform_%s"), **Item)))
+					.Image(FEditorStyle::GetBrush(PlatformInfo->GetIconStyleName(PlatformInfo::EPlatformIconSize::Normal)))
 			]
 
 		+ SHorizontalBox::Slot()
diff --git a/Engine/Source/Developer/DeviceManager/Private/Widgets/Browser/SDeviceBrowserDeviceListRow.h b/Engine/Source/Developer/DeviceManager/Private/Widgets/Browser/SDeviceBrowserDeviceListRow.h
--- a/Engine/Source/Developer/DeviceManager/Private/Widgets/Browser/SDeviceBrowserDeviceListRow.h
+++ b/Engine/Source/Developer/DeviceManager/Private/Widgets/Browser/SDeviceBrowserDeviceListRow.h
@@ -65,11 +65,14 @@ public:
 		}
 		else if (ColumnName == TEXT("Icon"))
 		{
+			const PlatformInfo::FPlatformInfo* const PlatformInfo = PlatformInfo::FindPlatformInfo(*DeviceService->GetDeviceId().GetPlatformName());
+			check(PlatformInfo);
+
 			return SNew(SBox)
 				.Padding(FMargin(4.0f, 0.0f))
 				[
 					SNew(SImage)
-						.Image(FEditorStyle::GetBrush(*FString::Printf(TEXT("Launcher.Platform_%s"), *DeviceService->GetDeviceId().GetPlatformName())))
+						.Image(FEditorStyle::GetBrush(PlatformInfo->GetIconStyleName(PlatformInfo::EPlatformIconSize::Normal)))
 				];
 		}
 		else if (ColumnName == TEXT("Name"))
diff --git a/Engine/Source/Developer/DeviceManager/Private/Widgets/Browser/SDeviceBrowserFilterBar.cpp b/Engine/Source/Developer/DeviceManager/Private/Widgets/Browser/SDeviceBrowserFilterBar.cpp
--- a/Engine/Source/Developer/DeviceManager/Private/Widgets/Browser/SDeviceBrowserFilterBar.cpp
+++ b/Engine/Source/Developer/DeviceManager/Private/Widgets/Browser/SDeviceBrowserFilterBar.cpp
@@ -113,6 +113,9 @@ FString SDeviceBrowserFilterBar::HandlePlatformListRowText( TSharedPtr<FString>
 BEGIN_SLATE_FUNCTION_BUILD_OPTIMIZATION
 TSharedRef<ITableRow> SDeviceBrowserFilterBar::HandlePlatformListViewGenerateRow( TSharedPtr<FString> PlatformName, const TSharedRef<STableViewBase>& OwnerTable )
 {
+	const PlatformInfo::FPlatformInfo* const PlatformInfo = PlatformInfo::FindPlatformInfo(**PlatformName);
+	check(PlatformInfo);
+
 	return SNew(STableRow<TSharedPtr<FString> >, OwnerTable)
 		.Content()
 		[
@@ -128,7 +131,7 @@ TSharedRef<ITableRow> SDeviceBrowserFilterBar::HandlePlatformListViewGenerateRow
 						.AutoWidth()
 						[
 							SNew(SImage)
-								.Image(FEditorStyle::GetBrush(*FString::Printf(TEXT("Launcher.Platform_%s"), **PlatformName)))
+								.Image(FEditorStyle::GetBrush(PlatformInfo->GetIconStyleName(PlatformInfo::EPlatformIconSize::Normal)))
 						]
 
 					+ SHorizontalBox::Slot()
diff --git a/Engine/Source/Developer/DeviceManager/Private/Widgets/Shared/SDeviceQuickInfo.h b/Engine/Source/Developer/DeviceManager/Private/Widgets/Shared/SDeviceQuickInfo.h
--- a/Engine/Source/Developer/DeviceManager/Private/Widgets/Shared/SDeviceQuickInfo.h
+++ b/Engine/Source/Developer/DeviceManager/Private/Widgets/Shared/SDeviceQuickInfo.h
@@ -245,7 +245,10 @@ private:
 	{
 		if (DeviceService.IsValid())
 		{
-			return FEditorStyle::GetBrush(*FString::Printf(TEXT("Launcher.Platform_%s.XLarge"), *DeviceService->GetDeviceId().GetPlatformName()));
+			const PlatformInfo::FPlatformInfo* const PlatformInfo = PlatformInfo::FindPlatformInfo(*DeviceService->GetDeviceId().GetPlatformName());
+			check(PlatformInfo);
+
+			return FEditorStyle::GetBrush(PlatformInfo->GetIconStyleName(PlatformInfo::EPlatformIconSize::XLarge));
 		}
 
 		return NULL;
diff --git a/Engine/Source/Developer/HTML5/HTML5TargetPlatform/HTML5TargetPlatform.Build.cs b/Engine/Source/Developer/HTML5/HTML5TargetPlatform/HTML5TargetPlatform.Build.cs
--- a/Engine/Source/Developer/HTML5/HTML5TargetPlatform/HTML5TargetPlatform.Build.cs
+++ b/Engine/Source/Developer/HTML5/HTML5TargetPlatform/HTML5TargetPlatform.Build.cs
@@ -11,13 +11,13 @@ public class HTML5TargetPlatform : ModuleRules
 	{
 		BinariesSubFolder = "HTML5";
 
-		PrivateIncludePathModuleNames.Add("TargetPlatform");
-
 		PrivateDependencyModuleNames.AddRange(
 			new string[] {
 				"Core",
 				"CoreUObject",
-				"Sockets"
+				"Sockets",
+				"TargetPlatform",
+				"DesktopPlatform",
 			}
 		);
 
diff --git a/Engine/Source/Developer/HTML5/HTML5TargetPlatform/Private/HTML5TargetPlatform.cpp b/Engine/Source/Developer/HTML5/HTML5TargetPlatform/Private/HTML5TargetPlatform.cpp
--- a/Engine/Source/Developer/HTML5/HTML5TargetPlatform/Private/HTML5TargetPlatform.cpp
+++ b/Engine/Source/Developer/HTML5/HTML5TargetPlatform/Private/HTML5TargetPlatform.cpp
@@ -61,22 +61,6 @@ ITargetDevicePtr FHTML5TargetPlatform::GetDevice( const FTargetDeviceId& DeviceI
 }
 
 
-FString FHTML5TargetPlatform::GetIconPath( ETargetPlatformIcons::IconType IconType ) const
-{
-	switch (IconType)
-	{
-	case ETargetPlatformIcons::Normal:
-		return FString(TEXT("Launcher/HTML5/Platform_HTML5_24x"));
-
-	case ETargetPlatformIcons::Large:
-	case ETargetPlatformIcons::XLarge:
-		return FString(TEXT("Launcher/HTML5/Platform_HTML5_128x"));
-	}
-
-	return FString();
-}
-
-
 bool FHTML5TargetPlatform::IsRunningPlatform( ) const
 {
 	return false; // but this will never be called because this platform doesn't run the target platform framework
diff --git a/Engine/Source/Developer/HTML5/HTML5TargetPlatform/Private/HTML5TargetPlatform.h b/Engine/Source/Developer/HTML5/HTML5TargetPlatform/Private/HTML5TargetPlatform.h
--- a/Engine/Source/Developer/HTML5/HTML5TargetPlatform/Private/HTML5TargetPlatform.h
+++ b/Engine/Source/Developer/HTML5/HTML5TargetPlatform/Private/HTML5TargetPlatform.h
@@ -42,8 +42,6 @@ public:
 
 	virtual ITargetDevicePtr GetDevice( const FTargetDeviceId& DeviceId ) OVERRIDE;
 
-	virtual FString GetIconPath( ETargetPlatformIcons::IconType IconType ) const OVERRIDE;
-
 	virtual bool IsRunningPlatform( ) const OVERRIDE;
 
 #if WITH_ENGINE
diff --git a/Engine/Source/Developer/HTML5/HTML5TargetPlatform/Private/HTML5TargetPlatformPrivatePCH.h b/Engine/Source/Developer/HTML5/HTML5TargetPlatform/Private/HTML5TargetPlatformPrivatePCH.h
--- a/Engine/Source/Developer/HTML5/HTML5TargetPlatform/Private/HTML5TargetPlatformPrivatePCH.h
+++ b/Engine/Source/Developer/HTML5/HTML5TargetPlatform/Private/HTML5TargetPlatformPrivatePCH.h
@@ -19,6 +19,7 @@
 #endif
 
 #include "TargetPlatform.h"
+#include "TargetPlatformBase.h"
 
 
 /* Private includes
diff --git a/Engine/Source/Developer/LauncherServices/LauncherServices.Build.cs b/Engine/Source/Developer/LauncherServices/LauncherServices.Build.cs
--- a/Engine/Source/Developer/LauncherServices/LauncherServices.Build.cs
+++ b/Engine/Source/Developer/LauncherServices/LauncherServices.Build.cs
@@ -30,6 +30,7 @@ namespace UnrealBuildTool.Rules
 			PrivateIncludePathModuleNames.AddRange(
 				new string[] {
 					"Messaging",
+					"DesktopPlatform",
 				}
 			);
 
diff --git a/Engine/Source/Developer/LauncherServices/Private/Launcher/LauncherUATTask.h b/Engine/Source/Developer/LauncherServices/Private/Launcher/LauncherUATTask.h
--- a/Engine/Source/Developer/LauncherServices/Private/Launcher/LauncherUATTask.h
+++ b/Engine/Source/Developer/LauncherServices/Private/Launcher/LauncherUATTask.h
@@ -6,6 +6,7 @@
 
 #pragma once
 
+#include "PlatformInfo.h"
 
 /**
  * class for UAT launcher tasks.
@@ -53,28 +54,29 @@ protected:
 		FString Executable = TEXT("RunUAT.bat");
 #endif
 
+		const PlatformInfo::FPlatformInfo& PlatformInfo = TargetPlatform.GetPlatformInfo();
+
 		// switch server and no editor platforms to the proper type
-		FString Platform = TargetPlatform.PlatformName();
-		if (TargetPlatform.PlatformName() == TEXT("LinuxServer") || TargetPlatform.PlatformName() == TEXT("LinuxNoEditor"))
+		FName PlatformName = PlatformInfo.TargetPlatformName;
+		if (PlatformName == FName("LinuxServer") || PlatformName == FName("LinuxNoEditor"))
 		{
-			Platform = TEXT("Linux");
+			PlatformName = FName("Linux");
 		}
-		else if (TargetPlatform.PlatformName() == TEXT("WindowsServer") || TargetPlatform.PlatformName() == TEXT("WindowsNoEditor") || TargetPlatform.PlatformName() == TEXT("Windows"))
+		else if (PlatformName == FName("WindowsServer") || PlatformName == FName("WindowsNoEditor") || PlatformName == FName("Windows"))
 		{
-			Platform = TEXT("Win64");
+			PlatformName = FName("Win64");
 		}
-        else if (TargetPlatform.PlatformName() == TEXT("MacNoEditor"))
+        else if (PlatformName == FName("MacNoEditor"))
         {
-            Platform = TEXT("Mac");
+            PlatformName = FName("Mac");
         }
 
-		// for things like Android_ATC, pull it apart to be Android, with cook flavor ATC
-		FString CookFlavor;
-		int32 UnderscoreLoc;
-		if (Platform.FindChar(TEXT('_'), UnderscoreLoc))
+		// Append any extra UAT flags specified for this platform flavor
+		FString OptionalParams;
+		if (!PlatformInfo.UATCommandLine.IsEmpty())
 		{
-			CookFlavor = FString(TEXT(" -cookflavor=")) + Platform.Mid(UnderscoreLoc + 1);
-			Platform = Platform.Mid(0, UnderscoreLoc);
+			OptionalParams += TEXT(" ");
+			OptionalParams += PlatformInfo.UATCommandLine;
 		}
 
         // check for rocket
@@ -90,13 +92,13 @@ protected:
 		ProjectPath = FPaths::ConvertRelativePathToFull(ProjectPath);
 		CommandLine = FString::Printf(TEXT("BuildCookRun -project=\"%s\" -noP4 -platform=%s -clientconfig=%s -serverconfig=%s"),
 			*ProjectPath,
-			*Platform,
+			*PlatformName.ToString(),
 			*ConfigStrings[ChainState.Profile->GetBuildConfiguration()],
 			*ConfigStrings[ChainState.Profile->GetBuildConfiguration()]);
 
 		CommandLine += NoCompile;
         CommandLine += Rocket;
-		CommandLine += CookFlavor;
+		CommandLine += OptionalParams;
 
 		// specialized command arguments for this particular task
 		CommandLine += TaskCommand->GetArguments(ChainState);
diff --git a/Engine/Source/Developer/Linux/LinuxNoEditorTargetPlatform/LinuxNoEditorTargetPlatform.Build.cs b/Engine/Source/Developer/Linux/LinuxNoEditorTargetPlatform/LinuxNoEditorTargetPlatform.Build.cs
--- a/Engine/Source/Developer/Linux/LinuxNoEditorTargetPlatform/LinuxNoEditorTargetPlatform.Build.cs
+++ b/Engine/Source/Developer/Linux/LinuxNoEditorTargetPlatform/LinuxNoEditorTargetPlatform.Build.cs
@@ -11,7 +11,8 @@ public class LinuxNoEditorTargetPlatform : ModuleRules
         PrivateDependencyModuleNames.AddRange(
             new string[] {
 				"Core",
-				"TargetPlatform"
+				"TargetPlatform",
+				"DesktopPlatform",
 			}
         );
 
diff --git a/Engine/Source/Developer/Linux/LinuxNoEditorTargetPlatform/Private/LinuxNoEditorTargetPlatformPrivatePCH.h b/Engine/Source/Developer/Linux/LinuxNoEditorTargetPlatform/Private/LinuxNoEditorTargetPlatformPrivatePCH.h
--- a/Engine/Source/Developer/Linux/LinuxNoEditorTargetPlatform/Private/LinuxNoEditorTargetPlatformPrivatePCH.h
+++ b/Engine/Source/Developer/Linux/LinuxNoEditorTargetPlatform/Private/LinuxNoEditorTargetPlatformPrivatePCH.h
@@ -20,6 +20,7 @@
 #endif
 
 #include "TargetPlatform.h"
+#include "TargetPlatformBase.h"
 
 
 /* Private includes
diff --git a/Engine/Source/Developer/Linux/LinuxServerTargetPlatform/LinuxServerTargetPlatform.Build.cs b/Engine/Source/Developer/Linux/LinuxServerTargetPlatform/LinuxServerTargetPlatform.Build.cs
--- a/Engine/Source/Developer/Linux/LinuxServerTargetPlatform/LinuxServerTargetPlatform.Build.cs
+++ b/Engine/Source/Developer/Linux/LinuxServerTargetPlatform/LinuxServerTargetPlatform.Build.cs
@@ -11,7 +11,8 @@ public class LinuxServerTargetPlatform : ModuleRules
         PrivateDependencyModuleNames.AddRange(
             new string[] {
 				"Core",
-				"TargetPlatform"
+				"TargetPlatform",
+				"DesktopPlatform",
 			}
         );
 
diff --git a/Engine/Source/Developer/Linux/LinuxServerTargetPlatform/Private/LinuxServerTargetPlatformPrivatePCH.h b/Engine/Source/Developer/Linux/LinuxServerTargetPlatform/Private/LinuxServerTargetPlatformPrivatePCH.h
--- a/Engine/Source/Developer/Linux/LinuxServerTargetPlatform/Private/LinuxServerTargetPlatformPrivatePCH.h
+++ b/Engine/Source/Developer/Linux/LinuxServerTargetPlatform/Private/LinuxServerTargetPlatformPrivatePCH.h
@@ -20,6 +20,7 @@
 #endif
 
 #include "TargetPlatform.h"
+#include "TargetPlatformBase.h"
 
 
 /* Private includes
diff --git a/Engine/Source/Developer/Linux/LinuxTargetPlatform/LinuxTargetPlatform.Build.cs b/Engine/Source/Developer/Linux/LinuxTargetPlatform/LinuxTargetPlatform.Build.cs
--- a/Engine/Source/Developer/Linux/LinuxTargetPlatform/LinuxTargetPlatform.Build.cs
+++ b/Engine/Source/Developer/Linux/LinuxTargetPlatform/LinuxTargetPlatform.Build.cs
@@ -18,7 +18,8 @@ public class LinuxTargetPlatform : ModuleRules
             new string[] {
 				"Core",
 				"CoreUObject",
-				"TargetPlatform"
+				"TargetPlatform",
+				"DesktopPlatform",
 			}
         );
 
diff --git a/Engine/Source/Developer/Linux/LinuxTargetPlatform/Private/LinuxTargetPlatform.h b/Engine/Source/Developer/Linux/LinuxTargetPlatform/Private/LinuxTargetPlatform.h
--- a/Engine/Source/Developer/Linux/LinuxTargetPlatform/Private/LinuxTargetPlatform.h
+++ b/Engine/Source/Developer/Linux/LinuxTargetPlatform/Private/LinuxTargetPlatform.h
@@ -86,21 +86,6 @@ public:
 		return MakeShareable(new FLinuxTargetDevice(*this, UATFriendlyId, DeviceId.GetDeviceName()));
 	}
 
-	virtual FString GetIconPath( ETargetPlatformIcons::IconType IconType ) const OVERRIDE
-	{
-		switch (IconType)
-		{
-		case ETargetPlatformIcons::Normal:
-			return FString(TEXT("Launcher/Linux/Platform_Linux_24x"));
-
-		case ETargetPlatformIcons::Large:
-		case ETargetPlatformIcons::XLarge:
-			return FString(TEXT("Launcher/Linux/Platform_Linux_128x"));
-		}
-
-		return FString();
-	}
-
 	virtual bool IsRunningPlatform( ) const OVERRIDE
 	{
 		return false;
diff --git a/Engine/Source/Developer/Linux/LinuxTargetPlatform/Private/LinuxTargetPlatformPrivatePCH.h b/Engine/Source/Developer/Linux/LinuxTargetPlatform/Private/LinuxTargetPlatformPrivatePCH.h
--- a/Engine/Source/Developer/Linux/LinuxTargetPlatform/Private/LinuxTargetPlatformPrivatePCH.h
+++ b/Engine/Source/Developer/Linux/LinuxTargetPlatform/Private/LinuxTargetPlatformPrivatePCH.h
@@ -21,6 +21,7 @@
 #endif
 
 #include "TargetPlatform.h"
+#include "TargetPlatformBase.h"
 
 
 /* Private includes
diff --git a/Engine/Source/Developer/Mac/MacClientTargetPlatform/MacClientTargetPlatform.Build.cs b/Engine/Source/Developer/Mac/MacClientTargetPlatform/MacClientTargetPlatform.Build.cs
--- a/Engine/Source/Developer/Mac/MacClientTargetPlatform/MacClientTargetPlatform.Build.cs
+++ b/Engine/Source/Developer/Mac/MacClientTargetPlatform/MacClientTargetPlatform.Build.cs
@@ -9,7 +9,8 @@ public class MacClientTargetPlatform : ModuleRules
 		PrivateDependencyModuleNames.AddRange(
 			new string[] {
 				"Core",
-				"TargetPlatform"
+				"TargetPlatform",
+				"DesktopPlatform",
 			}
 		);
 
diff --git a/Engine/Source/Developer/Mac/MacClientTargetPlatform/Private/MacClientTargetPlatformPrivatePCH.h b/Engine/Source/Developer/Mac/MacClientTargetPlatform/Private/MacClientTargetPlatformPrivatePCH.h
--- a/Engine/Source/Developer/Mac/MacClientTargetPlatform/Private/MacClientTargetPlatformPrivatePCH.h
+++ b/Engine/Source/Developer/Mac/MacClientTargetPlatform/Private/MacClientTargetPlatformPrivatePCH.h
@@ -20,6 +20,7 @@
 #endif
 
 #include "TargetPlatform.h"
+#include "TargetPlatformBase.h"
 #include "LocalMacTargetDevice.h"
 #include "GenericMacTargetPlatform.h"
 
diff --git a/Engine/Source/Developer/Mac/MacNoEditorTargetPlatform/MacNoEditorTargetPlatform.Build.cs b/Engine/Source/Developer/Mac/MacNoEditorTargetPlatform/MacNoEditorTargetPlatform.Build.cs
--- a/Engine/Source/Developer/Mac/MacNoEditorTargetPlatform/MacNoEditorTargetPlatform.Build.cs
+++ b/Engine/Source/Developer/Mac/MacNoEditorTargetPlatform/MacNoEditorTargetPlatform.Build.cs
@@ -9,7 +9,8 @@ public class MacNoEditorTargetPlatform : ModuleRules
 		PrivateDependencyModuleNames.AddRange(
 			new string[] {
 				"Core",
-				"TargetPlatform"
+				"TargetPlatform",
+				"DesktopPlatform",
 			}
 		);
 
diff --git a/Engine/Source/Developer/Mac/MacNoEditorTargetPlatform/Private/MacNoEditorTargetPlatformPrivatePCH.h b/Engine/Source/Developer/Mac/MacNoEditorTargetPlatform/Private/MacNoEditorTargetPlatformPrivatePCH.h
--- a/Engine/Source/Developer/Mac/MacNoEditorTargetPlatform/Private/MacNoEditorTargetPlatformPrivatePCH.h
+++ b/Engine/Source/Developer/Mac/MacNoEditorTargetPlatform/Private/MacNoEditorTargetPlatformPrivatePCH.h
@@ -20,6 +20,7 @@
 #endif
 
 #include "TargetPlatform.h"
+#include "TargetPlatformBase.h"
 #include "LocalMacTargetDevice.h"
 #include "GenericMacTargetPlatform.h"
 
diff --git a/Engine/Source/Developer/Mac/MacServerTargetPlatform/MacServerTargetPlatform.Build.cs b/Engine/Source/Developer/Mac/MacServerTargetPlatform/MacServerTargetPlatform.Build.cs
--- a/Engine/Source/Developer/Mac/MacServerTargetPlatform/MacServerTargetPlatform.Build.cs
+++ b/Engine/Source/Developer/Mac/MacServerTargetPlatform/MacServerTargetPlatform.Build.cs
@@ -9,7 +9,8 @@ public class MacServerTargetPlatform : ModuleRules
 		PrivateDependencyModuleNames.AddRange(
 			new string[] {
 				"Core",
-				"TargetPlatform"
+				"TargetPlatform",
+				"DesktopPlatform",
 			}
 		);
 
diff --git a/Engine/Source/Developer/Mac/MacServerTargetPlatform/Private/MacServerTargetPlatformPrivatePCH.h b/Engine/Source/Developer/Mac/MacServerTargetPlatform/Private/MacServerTargetPlatformPrivatePCH.h
--- a/Engine/Source/Developer/Mac/MacServerTargetPlatform/Private/MacServerTargetPlatformPrivatePCH.h
+++ b/Engine/Source/Developer/Mac/MacServerTargetPlatform/Private/MacServerTargetPlatformPrivatePCH.h
@@ -20,6 +20,7 @@
 #endif
 
 #include "TargetPlatform.h"
+#include "TargetPlatformBase.h"
 #include "LocalMacTargetDevice.h"
 #include "GenericMacTargetPlatform.h"
 
diff --git a/Engine/Source/Developer/Mac/MacTargetPlatform/MacTargetPlatform.Build.cs b/Engine/Source/Developer/Mac/MacTargetPlatform/MacTargetPlatform.Build.cs
--- a/Engine/Source/Developer/Mac/MacTargetPlatform/MacTargetPlatform.Build.cs
+++ b/Engine/Source/Developer/Mac/MacTargetPlatform/MacTargetPlatform.Build.cs
@@ -10,7 +10,8 @@ public class MacTargetPlatform : ModuleRules
 			new string[] {
 				"Core",
 				"CoreUObject",
-				"TargetPlatform"
+				"TargetPlatform",
+				"DesktopPlatform",
 			}
 		);
 
diff --git a/Engine/Source/Developer/Mac/MacTargetPlatform/Private/GenericMacTargetPlatform.h b/Engine/Source/Developer/Mac/MacTargetPlatform/Private/GenericMacTargetPlatform.h
--- a/Engine/Source/Developer/Mac/MacTargetPlatform/Private/GenericMacTargetPlatform.h
+++ b/Engine/Source/Developer/Mac/MacTargetPlatform/Private/GenericMacTargetPlatform.h
@@ -72,52 +72,6 @@ public:
 		return NULL;
 	}
 
-	virtual FString GetIconPath( ETargetPlatformIcons::IconType IconType ) const OVERRIDE
-	{
-		if (IS_DEDICATED_SERVER)
-		{
-			switch (IconType)
-			{
-			case ETargetPlatformIcons::Normal:
-//				return FString(TEXT("Launcher/Mac/Platform_MacServer_24x"));
-				return FString(TEXT("Launcher/Mac/Platform_Mac_24x"));
-
-			case ETargetPlatformIcons::Large:
-			case ETargetPlatformIcons::XLarge:
-//				return FString(TEXT("Launcher/Mac/Platform_MacServer_128x"));
-				return FString(TEXT("Launcher/Mac/Platform_Mac_128x"));
-			}
-		}
-		else if (UE_GAME)
-		{
-			switch (IconType)
-			{
-			case ETargetPlatformIcons::Normal:
-//				return FString(TEXT("Launcher/Mac/Platform_MacNoEditor_24x"));
-				return FString(TEXT("Launcher/Mac/Platform_Mac_24x"));
-
-			case ETargetPlatformIcons::Large:
-			case ETargetPlatformIcons::XLarge:
-//				return FString(TEXT("Launcher/Mac/Platform_MacNoEditor_128x"));
-				return FString(TEXT("Launcher/Mac/Platform_Mac_128x"));
-			}
-		}
-		else
-		{
-			switch (IconType)
-			{
-			case ETargetPlatformIcons::Normal:
-				return FString(TEXT("Launcher/Mac/Platform_Mac_24x"));
-
-			case ETargetPlatformIcons::Large:
-			case ETargetPlatformIcons::XLarge:
-				return FString(TEXT("Launcher/Mac/Platform_Mac_128x"));
-			}
-		}
-
-		return FString(TEXT(""));
-	}
-
 	virtual bool IsRunningPlatform( ) const OVERRIDE
 	{
 		// Must be Mac platform as editor for this to be considered a running platform
diff --git a/Engine/Source/Developer/Mac/MacTargetPlatform/Private/MacTargetPlatformPrivatePCH.h b/Engine/Source/Developer/Mac/MacTargetPlatform/Private/MacTargetPlatformPrivatePCH.h
--- a/Engine/Source/Developer/Mac/MacTargetPlatform/Private/MacTargetPlatformPrivatePCH.h
+++ b/Engine/Source/Developer/Mac/MacTargetPlatform/Private/MacTargetPlatformPrivatePCH.h
@@ -22,6 +22,7 @@
 #endif
 
 #include "TargetPlatform.h"
+#include "TargetPlatformBase.h"
 
 /* Private includes
  *****************************************************************************/
diff --git a/Engine/Source/Developer/SessionFrontend/Private/SessionFrontendPrivatePCH.h b/Engine/Source/Developer/SessionFrontend/Private/SessionFrontendPrivatePCH.h
--- a/Engine/Source/Developer/SessionFrontend/Private/SessionFrontendPrivatePCH.h
+++ b/Engine/Source/Developer/SessionFrontend/Private/SessionFrontendPrivatePCH.h
@@ -20,6 +20,7 @@
 #include "Slate.h"
 #include "TargetDeviceServices.h"
 #include "TargetPlatform.h"
+#include "PlatformInfo.h"
 
 // @todo gmp: remove these dependencies by making the session front-end extensible
 #include "AutomationWindow.h"
diff --git a/Engine/Source/Developer/SessionFrontend/Private/Widgets/Browser/SSessionBrowserInstanceListRow.h b/Engine/Source/Developer/SessionFrontend/Private/Widgets/Browser/SSessionBrowserInstanceListRow.h
--- a/Engine/Source/Developer/SessionFrontend/Private/Widgets/Browser/SSessionBrowserInstanceListRow.h
+++ b/Engine/Source/Developer/SessionFrontend/Private/Widgets/Browser/SSessionBrowserInstanceListRow.h
@@ -98,6 +98,9 @@ public:
 		}
 		else if (ColumnName == "Platform")
 		{
+			const PlatformInfo::FPlatformInfo* const PlatformInfo = PlatformInfo::FindPlatformInfo(*InstanceInfo->GetPlatformName());
+			check(PlatformInfo);
+
 			return SNew(SHorizontalBox)
 
 			+ SHorizontalBox::Slot()
@@ -105,7 +108,7 @@ public:
 				.VAlign(VAlign_Center)
 				[
 					SNew(SImage)
-						.Image(FEditorStyle::GetBrush(*FString::Printf(TEXT("Launcher.Platform_%s"), *InstanceInfo->GetPlatformName())))
+						.Image(FEditorStyle::GetBrush(PlatformInfo->GetIconStyleName(PlatformInfo::EPlatformIconSize::Normal)))
 				]
 
 			+ SHorizontalBox::Slot()
diff --git a/Engine/Source/Developer/SessionLauncher/Private/SessionLauncherPrivatePCH.h b/Engine/Source/Developer/SessionLauncher/Private/SessionLauncherPrivatePCH.h
--- a/Engine/Source/Developer/SessionLauncher/Private/SessionLauncherPrivatePCH.h
+++ b/Engine/Source/Developer/SessionLauncher/Private/SessionLauncherPrivatePCH.h
@@ -19,6 +19,7 @@
 #include "ModuleManager.h"
 #include "TargetDeviceServices.h"
 #include "TargetPlatform.h"
+#include "PlatformInfo.h"
 #include "TextFilter.h"
 
 
diff --git a/Engine/Source/Developer/SessionLauncher/Private/Widgets/Deploy/SSessionLauncherDeviceListRow.h b/Engine/Source/Developer/SessionLauncher/Private/Widgets/Deploy/SSessionLauncherDeviceListRow.h
--- a/Engine/Source/Developer/SessionLauncher/Private/Widgets/Deploy/SSessionLauncherDeviceListRow.h
+++ b/Engine/Source/Developer/SessionLauncher/Private/Widgets/Deploy/SSessionLauncherDeviceListRow.h
@@ -171,7 +171,10 @@ private:
 	// Callback for getting the icon image of the device.
 	const FSlateBrush* HandleDeviceImage( ) const
 	{
-		return FEditorStyle::GetBrush(*FString::Printf(TEXT("Launcher.Platform_%s"), *DeviceProxy->GetPlatformName()));
+		const PlatformInfo::FPlatformInfo* const PlatformInfo = PlatformInfo::FindPlatformInfo(*DeviceProxy->GetPlatformName());
+		check(PlatformInfo);
+
+		return FEditorStyle::GetBrush(PlatformInfo->GetIconStyleName(PlatformInfo::EPlatformIconSize::Normal));
 	}
 
 	// Callback for getting the friendly name.
diff --git a/Engine/Source/Developer/TargetPlatform/Public/Common/TargetPlatformBase.h b/Engine/Source/Developer/TargetPlatform/Public/Common/TargetPlatformBase.h
--- a/Engine/Source/Developer/TargetPlatform/Public/Common/TargetPlatformBase.h
+++ b/Engine/Source/Developer/TargetPlatform/Public/Common/TargetPlatformBase.h
@@ -6,6 +6,7 @@
 
 #pragma once
 
+#include "PlatformInfo.h"
 
 /**
  * Base class for target platforms.
@@ -22,6 +23,16 @@ public:
 		return false;
 	}
 
+	virtual FText DisplayName( ) const OVERRIDE
+	{
+		return PlatformInfo->DisplayName;
+	}
+
+	virtual const PlatformInfo::FPlatformInfo& GetPlatformInfo( ) const OVERRIDE
+	{
+		return *PlatformInfo;
+	}
+
 #if WITH_ENGINE
 	virtual void GetReflectionCaptureFormats( TArray<FName>& OutFormats ) const
 	{
@@ -156,6 +167,16 @@ public:
 	}
 
 	// End ITargetPlatform interface
+
+protected:
+	FTargetPlatformBase(const PlatformInfo::FPlatformInfo *const InPlatformInfo)
+		: PlatformInfo(InPlatformInfo)
+	{
+		check(PlatformInfo);
+	}
+
+	/** Information about this platform */
+	const PlatformInfo::FPlatformInfo *PlatformInfo;
 };
 
 
@@ -174,6 +195,7 @@ public:
 	 * Default constructor.
 	 */
 	TTargetPlatformBase( )
+		: FTargetPlatformBase( PlatformInfo::FindPlatformInfo(TPlatformProperties::PlatformName()) )
 	{
 		// HasEditorOnlyData and RequiresCookedData are mutually exclusive.
 		check(TPlatformProperties::HasEditorOnlyData() != TPlatformProperties::RequiresCookedData());
@@ -183,11 +205,6 @@ public:
 
 	// Begin ITargetPlatform interface
 
-	virtual FText DisplayName( ) const OVERRIDE
-	{
-		return FText::FromString(TPlatformProperties::DisplayName());
-	}
-
 	virtual bool HasEditorOnlyData( ) const OVERRIDE
 	{
 		return TPlatformProperties::HasEditorOnlyData();
diff --git a/Engine/Source/Developer/TargetPlatform/Public/Interfaces/ITargetPlatform.h b/Engine/Source/Developer/TargetPlatform/Public/Interfaces/ITargetPlatform.h
--- a/Engine/Source/Developer/TargetPlatform/Public/Interfaces/ITargetPlatform.h
+++ b/Engine/Source/Developer/TargetPlatform/Public/Interfaces/ITargetPlatform.h
@@ -7,6 +7,13 @@
 #pragma once
 
 
+namespace PlatformInfo
+{
+	// Forward declare type from DesktopPlatform rather than add an include dependency to everything using ITargetPlatform
+	struct FPlatformInfo;
+}
+
+
 namespace ETargetPlatformBuildArtifacts
 {
 	/**
@@ -76,25 +83,6 @@ namespace ETargetPlatformFeatures
 };
 
 
-namespace ETargetPlatformIcons
-{
-	/**
-	 * Enumerates platform icon types.
-	 */
-	enum IconType
-	{
-		/** Normal sized icon. */
-		Normal,
-
-		/** Large icon. */
-		Large,
-
-		/** Extra large icon. */
-		XLarge
-	};
-}
-
-
 /**
  * ITargetPlatform, abstraction for cooking platforms and enumerating actual target devices
 **/
@@ -118,6 +106,11 @@ public:
 	 */
 	virtual FText DisplayName( ) const = 0;
 
+	/**
+	 * Returns the information about this platform
+	 */
+	virtual const PlatformInfo::FPlatformInfo& GetPlatformInfo( ) const = 0;
+
 	/**
 	 * Gets the platform's ini name (so an offline tool can load the ini for the given target platform)
 	 *
@@ -170,14 +163,6 @@ public:
 	 */
 	virtual ITargetDevicePtr GetDevice( const FTargetDeviceId& DeviceId ) = 0;
 
-	/**
-	 * Gets the path to the platform's icon.
-	 *
-	 * @param IconType The type of icon to get (i.e. Small or Large).
-	 * @return The path to the icon.
-	 */
-	virtual FString GetIconPath( ETargetPlatformIcons::IconType IconType ) const = 0;
-
 	/**
 	 * Checks whether this platform has only editor data (typically desktop platforms).
 	 *
diff --git a/Engine/Source/Developer/TargetPlatform/Public/TargetPlatform.h b/Engine/Source/Developer/TargetPlatform/Public/TargetPlatform.h
--- a/Engine/Source/Developer/TargetPlatform/Public/TargetPlatform.h
+++ b/Engine/Source/Developer/TargetPlatform/Public/TargetPlatform.h
@@ -34,9 +34,3 @@
 
 #include "ITextureFormat.h"
 #include "ITextureFormatModule.h"
-
-
-/* Common
- *****************************************************************************/
-
-#include "TargetPlatformBase.h"
diff --git a/Engine/Source/Developer/TargetPlatform/TargetPlatform.Build.cs b/Engine/Source/Developer/TargetPlatform/TargetPlatform.Build.cs
--- a/Engine/Source/Developer/TargetPlatform/TargetPlatform.Build.cs
+++ b/Engine/Source/Developer/TargetPlatform/TargetPlatform.Build.cs
@@ -9,6 +9,8 @@ public class TargetPlatform : ModuleRules
 	{
 		PrivateDependencyModuleNames.Add("Core");
 
+		PublicIncludePathModuleNames.Add("DesktopPlatform");
+
 		if (!UEBuildConfiguration.bBuildRequiresCookedData)
 		{
             // these are needed by multiple platform specific target platforms, so we make sure they are built with the base editor
diff --git a/Engine/Source/Developer/Windows/WindowsClientTargetPlatform/Private/WindowsClientTargetPlatformPrivatePCH.h b/Engine/Source/Developer/Windows/WindowsClientTargetPlatform/Private/WindowsClientTargetPlatformPrivatePCH.h
--- a/Engine/Source/Developer/Windows/WindowsClientTargetPlatform/Private/WindowsClientTargetPlatformPrivatePCH.h
+++ b/Engine/Source/Developer/Windows/WindowsClientTargetPlatform/Private/WindowsClientTargetPlatformPrivatePCH.h
@@ -20,7 +20,7 @@
 #endif
 
 #include "TargetPlatform.h"
-
+#include "TargetPlatformBase.h"
 
 #include "AllowWindowsPlatformTypes.h"
 	#include <TlHelp32.h>
diff --git a/Engine/Source/Developer/Windows/WindowsClientTargetPlatform/WindowsClientTargetPlatform.Build.cs b/Engine/Source/Developer/Windows/WindowsClientTargetPlatform/WindowsClientTargetPlatform.Build.cs
--- a/Engine/Source/Developer/Windows/WindowsClientTargetPlatform/WindowsClientTargetPlatform.Build.cs
+++ b/Engine/Source/Developer/Windows/WindowsClientTargetPlatform/WindowsClientTargetPlatform.Build.cs
@@ -9,7 +9,8 @@ public class WindowsClientTargetPlatform : ModuleRules
 		PrivateDependencyModuleNames.AddRange(
 			new string[] {
 				"Core",
-				"TargetPlatform"
+				"TargetPlatform",
+				"DesktopPlatform",
 			}
 		);
 
diff --git a/Engine/Source/Developer/Windows/WindowsNoEditorTargetPlatform/Private/WindowsNoEditorTargetPlatformPrivatePCH.h b/Engine/Source/Developer/Windows/WindowsNoEditorTargetPlatform/Private/WindowsNoEditorTargetPlatformPrivatePCH.h
--- a/Engine/Source/Developer/Windows/WindowsNoEditorTargetPlatform/Private/WindowsNoEditorTargetPlatformPrivatePCH.h
+++ b/Engine/Source/Developer/Windows/WindowsNoEditorTargetPlatform/Private/WindowsNoEditorTargetPlatformPrivatePCH.h
@@ -20,6 +20,7 @@
 #endif
 
 #include "TargetPlatform.h"
+#include "TargetPlatformBase.h"
 
 #include "AllowWindowsPlatformTypes.h"
 	#include <TlHelp32.h>
diff --git a/Engine/Source/Developer/Windows/WindowsNoEditorTargetPlatform/WindowsNoEditorTargetPlatform.Build.cs b/Engine/Source/Developer/Windows/WindowsNoEditorTargetPlatform/WindowsNoEditorTargetPlatform.Build.cs
--- a/Engine/Source/Developer/Windows/WindowsNoEditorTargetPlatform/WindowsNoEditorTargetPlatform.Build.cs
+++ b/Engine/Source/Developer/Windows/WindowsNoEditorTargetPlatform/WindowsNoEditorTargetPlatform.Build.cs
@@ -9,7 +9,8 @@ public class WindowsNoEditorTargetPlatform : ModuleRules
 		PrivateDependencyModuleNames.AddRange(
 			new string[] {
 				"Core",
-				"TargetPlatform"
+				"TargetPlatform",
+				"DesktopPlatform",
 			}
 		);
 
diff --git a/Engine/Source/Developer/Windows/WindowsServerTargetPlatform/Private/WindowsServerTargetPlatformPrivatePCH.h b/Engine/Source/Developer/Windows/WindowsServerTargetPlatform/Private/WindowsServerTargetPlatformPrivatePCH.h
--- a/Engine/Source/Developer/Windows/WindowsServerTargetPlatform/Private/WindowsServerTargetPlatformPrivatePCH.h
+++ b/Engine/Source/Developer/Windows/WindowsServerTargetPlatform/Private/WindowsServerTargetPlatformPrivatePCH.h
@@ -20,7 +20,7 @@
 #endif
 
 #include "TargetPlatform.h"
-
+#include "TargetPlatformBase.h"
 
 #include "AllowWindowsPlatformTypes.h"
 	#include <TlHelp32.h>
diff --git a/Engine/Source/Developer/Windows/WindowsServerTargetPlatform/WindowsServerTargetPlatform.Build.cs b/Engine/Source/Developer/Windows/WindowsServerTargetPlatform/WindowsServerTargetPlatform.Build.cs
--- a/Engine/Source/Developer/Windows/WindowsServerTargetPlatform/WindowsServerTargetPlatform.Build.cs
+++ b/Engine/Source/Developer/Windows/WindowsServerTargetPlatform/WindowsServerTargetPlatform.Build.cs
@@ -9,7 +9,8 @@ public class WindowsServerTargetPlatform : ModuleRules
 		PrivateDependencyModuleNames.AddRange(
 			new string[] {
 				"Core",
-				"TargetPlatform"
+				"TargetPlatform",
+				"DesktopPlatform",
 			}
 		);
 
diff --git a/Engine/Source/Developer/Windows/WindowsTargetPlatform/Private/GenericWindowsTargetPlatform.h b/Engine/Source/Developer/Windows/WindowsTargetPlatform/Private/GenericWindowsTargetPlatform.h
--- a/Engine/Source/Developer/Windows/WindowsTargetPlatform/Private/GenericWindowsTargetPlatform.h
+++ b/Engine/Source/Developer/Windows/WindowsTargetPlatform/Private/GenericWindowsTargetPlatform.h
@@ -88,48 +88,6 @@ public:
 		return NULL;
 	}
 
-	virtual FString GetIconPath( ETargetPlatformIcons::IconType IconType ) const OVERRIDE
-	{
-		if (IS_DEDICATED_SERVER)
-		{
-			switch (IconType)
-			{
-			case ETargetPlatformIcons::Normal:
-				return FString(TEXT("Launcher/Windows/Platform_WindowsServer_24x"));
-
-			case ETargetPlatformIcons::Large:
-			case ETargetPlatformIcons::XLarge:
-				return FString(TEXT("Launcher/Windows/Platform_WindowsServer_128x"));
-			}
-		}
-		else if (UE_GAME)
-		{
-			switch (IconType)
-			{
-			case ETargetPlatformIcons::Normal:
-				return FString(TEXT("Launcher/Windows/Platform_WindowsNoEditor_24x"));
-
-			case ETargetPlatformIcons::Large:
-			case ETargetPlatformIcons::XLarge:
-				return FString(TEXT("Launcher/Windows/Platform_WindowsNoEditor_128x"));
-			}
-		}
-		else
-		{
-			switch (IconType)
-			{
-			case ETargetPlatformIcons::Normal:
-				return FString(TEXT("Launcher/Windows/Platform_Windows_24x"));
-
-			case ETargetPlatformIcons::Large:
-			case ETargetPlatformIcons::XLarge:
-				return FString(TEXT("Launcher/Windows/Platform_Windows_128x"));
-			}
-		}
-
-		return FString();
-	}
-
 	virtual bool IsRunningPlatform( ) const OVERRIDE
 	{
 		// Must be Windows platform as editor for this to be considered a running platform
diff --git a/Engine/Source/Developer/Windows/WindowsTargetPlatform/Private/WindowsTargetPlatformPrivatePCH.h b/Engine/Source/Developer/Windows/WindowsTargetPlatform/Private/WindowsTargetPlatformPrivatePCH.h
--- a/Engine/Source/Developer/Windows/WindowsTargetPlatform/Private/WindowsTargetPlatformPrivatePCH.h
+++ b/Engine/Source/Developer/Windows/WindowsTargetPlatform/Private/WindowsTargetPlatformPrivatePCH.h
@@ -21,6 +21,7 @@
 #endif
 
 #include "TargetPlatform.h"
+#include "TargetPlatformBase.h"
 
 #include "AllowWindowsPlatformTypes.h"
 	#include <TlHelp32.h>
diff --git a/Engine/Source/Developer/Windows/WindowsTargetPlatform/WindowsTargetPlatform.Build.cs b/Engine/Source/Developer/Windows/WindowsTargetPlatform/WindowsTargetPlatform.Build.cs
--- a/Engine/Source/Developer/Windows/WindowsTargetPlatform/WindowsTargetPlatform.Build.cs
+++ b/Engine/Source/Developer/Windows/WindowsTargetPlatform/WindowsTargetPlatform.Build.cs
@@ -10,7 +10,8 @@ public class WindowsTargetPlatform : ModuleRules
 			new string[] {
 				"Core",
 				"CoreUObject",
-				"TargetPlatform"
+				"TargetPlatform",
+				"DesktopPlatform",
 			}
 		);
 
diff --git a/Engine/Source/Developer/iOS/IOSTargetPlatform/IOSTargetPlatform.Build.cs b/Engine/Source/Developer/iOS/IOSTargetPlatform/IOSTargetPlatform.Build.cs
--- a/Engine/Source/Developer/iOS/IOSTargetPlatform/IOSTargetPlatform.Build.cs
+++ b/Engine/Source/Developer/iOS/IOSTargetPlatform/IOSTargetPlatform.Build.cs
@@ -14,6 +14,7 @@ public class IOSTargetPlatform : ModuleRules
 				"Core",
 				"CoreUObject",
 				"TargetPlatform",
+				"DesktopPlatform",
 				"LaunchDaemonMessages",
 			}
 		);
diff --git a/Engine/Source/Developer/iOS/IOSTargetPlatform/Private/IOSTargetPlatform.cpp b/Engine/Source/Developer/iOS/IOSTargetPlatform/Private/IOSTargetPlatform.cpp
--- a/Engine/Source/Developer/iOS/IOSTargetPlatform/Private/IOSTargetPlatform.cpp
+++ b/Engine/Source/Developer/iOS/IOSTargetPlatform/Private/IOSTargetPlatform.cpp
@@ -76,21 +76,6 @@ ITargetDevicePtr FIOSTargetPlatform::GetDevice( const FTargetDeviceId& DeviceId
 }
 
 
-FString FIOSTargetPlatform::GetIconPath( ETargetPlatformIcons::IconType InType ) const
-{
-	switch (InType)
-	{
-	case ETargetPlatformIcons::Normal:
-		return FString(TEXT("Launcher/iOS/Platform_iOS_24x"));
-
-	case ETargetPlatformIcons::Large:
-	case ETargetPlatformIcons::XLarge:
-		return FString(TEXT("Launcher/iOS/Platform_iOS_128x"));
-	}
-
-	return FString();
-}
-
 bool FIOSTargetPlatform::IsSdkInstalled(bool bProjectHasCode, FString& OutDocumentationPath) const
 {
 	bool biOSSDKInstalled = true; // @todo How do we check that the iOS SDK is installed when building from Windows? Is that even possible?
diff --git a/Engine/Source/Developer/iOS/IOSTargetPlatform/Private/IOSTargetPlatform.h b/Engine/Source/Developer/iOS/IOSTargetPlatform/Private/IOSTargetPlatform.h
--- a/Engine/Source/Developer/iOS/IOSTargetPlatform/Private/IOSTargetPlatform.h
+++ b/Engine/Source/Developer/iOS/IOSTargetPlatform/Private/IOSTargetPlatform.h
@@ -63,8 +63,6 @@ public:
 
 	virtual ITargetDevicePtr GetDevice( const FTargetDeviceId& DeviceId ) OVERRIDE;
 
-	virtual FString GetIconPath( ETargetPlatformIcons::IconType IconType ) const OVERRIDE;
-
 	virtual bool IsRunningPlatform( ) const OVERRIDE
 	{
 		#if PLATFORM_IOS && WITH_EDITOR
diff --git a/Engine/Source/Developer/iOS/IOSTargetPlatform/Private/IOSTargetPlatformPrivatePCH.h b/Engine/Source/Developer/iOS/IOSTargetPlatform/Private/IOSTargetPlatformPrivatePCH.h
--- a/Engine/Source/Developer/iOS/IOSTargetPlatform/Private/IOSTargetPlatformPrivatePCH.h
+++ b/Engine/Source/Developer/iOS/IOSTargetPlatform/Private/IOSTargetPlatformPrivatePCH.h
@@ -17,6 +17,7 @@
 #include "Messaging.h"
 #include "TargetDeviceServices.h"
 #include "TargetPlatform.h"
+#include "TargetPlatformBase.h"
 #include "IOSPlatformProperties.h"
 
 #if WITH_ENGINE
diff --git a/Engine/Source/Editor/DeviceProfileEditor/DeviceProfileEditor.Build.cs b/Engine/Source/Editor/DeviceProfileEditor/DeviceProfileEditor.Build.cs
--- a/Engine/Source/Editor/DeviceProfileEditor/DeviceProfileEditor.Build.cs
+++ b/Engine/Source/Editor/DeviceProfileEditor/DeviceProfileEditor.Build.cs
@@ -34,6 +34,7 @@ public class DeviceProfileEditor : ModuleRules
 				"PropertyEditor",
 				"SourceControl",
                 "TargetPlatform",
+				"DesktopPlatform",
 			}
 		);
 
diff --git a/Engine/Source/Editor/DeviceProfileEditor/Private/DetailsPanel/SDeviceProfileDetailsPanel.cpp b/Engine/Source/Editor/DeviceProfileEditor/Private/DetailsPanel/SDeviceProfileDetailsPanel.cpp
--- a/Engine/Source/Editor/DeviceProfileEditor/Private/DetailsPanel/SDeviceProfileDetailsPanel.cpp
+++ b/Engine/Source/Editor/DeviceProfileEditor/Private/DetailsPanel/SDeviceProfileDetailsPanel.cpp
@@ -89,12 +89,11 @@ void SDeviceProfileDetailsPanel::RefreshUI()
 
 	if( ViewingProfile.IsValid() )
 	{
-		FString DeviceProfileTypeIcon;
+		const FSlateBrush* DeviceProfileTypeIcon = FEditorStyle::GetDefaultBrush();
 		TArray<ITargetPlatform*> TargetPlatforms = GetTargetPlatformManager()->GetTargetPlatforms();
-		for (int32 TargetPlatformIndex = 0; TargetPlatformIndex < TargetPlatforms.Num(); ++TargetPlatformIndex)
+		if (TargetPlatforms.Num())
 		{
-			DeviceProfileTypeIcon = FString::Printf(TEXT("Launcher.Platform_%s"), *TargetPlatforms[TargetPlatformIndex]->PlatformName());
-			break;
+			DeviceProfileTypeIcon = FEditorStyle::GetBrush(TargetPlatforms[0]->GetPlatformInfo().GetIconStyleName(PlatformInfo::EPlatformIconSize::Normal));
 		}
 
 		SettingsView->SetObject(&*ViewingProfile);
@@ -115,7 +114,7 @@ void SDeviceProfileDetailsPanel::RefreshUI()
 					.Padding(4.0f, 0.0f, 2.0f, 0.0f)
 					[
 						SNew(SImage)
-						.Image(FEditorStyle::GetBrush(*DeviceProfileTypeIcon))
+						.Image(DeviceProfileTypeIcon)
 					]
 
 					+SHorizontalBox::Slot()
diff --git a/Engine/Source/Editor/DeviceProfileEditor/Private/DeviceProfileEditorPCH.h b/Engine/Source/Editor/DeviceProfileEditor/Private/DeviceProfileEditorPCH.h
--- a/Engine/Source/Editor/DeviceProfileEditor/Private/DeviceProfileEditorPCH.h
+++ b/Engine/Source/Editor/DeviceProfileEditor/Private/DeviceProfileEditorPCH.h
@@ -15,6 +15,7 @@
 #include "UnrealEd.h"
 #include "IPropertyTable.h"
 #include "TargetPlatform.h"
+#include "PlatformInfo.h"
 
 /* Private includes
 *****************************************************************************/
diff --git a/Engine/Source/Editor/DeviceProfileServices/DeviceProfileServices.Build.cs b/Engine/Source/Editor/DeviceProfileServices/DeviceProfileServices.Build.cs
--- a/Engine/Source/Editor/DeviceProfileServices/DeviceProfileServices.Build.cs
+++ b/Engine/Source/Editor/DeviceProfileServices/DeviceProfileServices.Build.cs
@@ -21,6 +21,7 @@ public class DeviceProfileServices : ModuleRules
 				
 				"CoreUObject",
 				"TargetPlatform",
+				"DesktopPlatform",
 				"UnrealEd",
 			}
 		);
diff --git a/Engine/Source/Editor/DeviceProfileServices/Private/DeviceProfileServicesPCH.h b/Engine/Source/Editor/DeviceProfileServices/Private/DeviceProfileServicesPCH.h
--- a/Engine/Source/Editor/DeviceProfileServices/Private/DeviceProfileServicesPCH.h
+++ b/Engine/Source/Editor/DeviceProfileServices/Private/DeviceProfileServicesPCH.h
@@ -15,6 +15,7 @@
 #include "Engine.h"
 #include "ModuleManager.h"
 #include "TargetDeviceServices.h"
+#include "PlatformInfo.h"
 
 
 /* Private includes
diff --git a/Engine/Source/Editor/DeviceProfileServices/Private/DeviceProfileServicesUIManager.cpp b/Engine/Source/Editor/DeviceProfileServices/Private/DeviceProfileServicesUIManager.cpp
--- a/Engine/Source/Editor/DeviceProfileServices/Private/DeviceProfileServicesUIManager.cpp
+++ b/Engine/Source/Editor/DeviceProfileServices/Private/DeviceProfileServicesUIManager.cpp
@@ -84,8 +84,7 @@ void FDeviceProfileServicesUIManager::CreatePlatformMap()
 	for (int32 Index = 0; Index < Platforms.Num(); ++Index)
 	{
 		PlatformList.Add(MakeShareable(new FString(Platforms[Index]->PlatformName())));
-		FName PlatformIcon = *FString::Printf(TEXT("Launcher.Platform_%s"), *Platforms[Index]->PlatformName());
-		DeviceTypeToIconMap.Add( Platforms[Index]->PlatformName(), PlatformIcon );
+		DeviceTypeToIconMap.Add( Platforms[Index]->PlatformName(), Platforms[Index]->GetPlatformInfo().GetIconStyleName(PlatformInfo::EPlatformIconSize::Normal) );
 	}
 }
 
diff --git a/Engine/Source/Editor/EditorStyle/EditorStyle.Build.cs b/Engine/Source/Editor/EditorStyle/EditorStyle.Build.cs
--- a/Engine/Source/Editor/EditorStyle/EditorStyle.Build.cs
+++ b/Engine/Source/Editor/EditorStyle/EditorStyle.Build.cs
@@ -17,7 +17,7 @@ public class EditorStyle : ModuleRules
 		PrivateDependencyModuleNames.AddRange(
 			new string[] {
 				"SlateCore",
-				"TargetPlatform",
+				"DesktopPlatform",
 			}
 		);
 
diff --git a/Engine/Source/Editor/EditorStyle/Private/EditorStylePrivatePCH.h b/Engine/Source/Editor/EditorStyle/Private/EditorStylePrivatePCH.h
--- a/Engine/Source/Editor/EditorStyle/Private/EditorStylePrivatePCH.h
+++ b/Engine/Source/Editor/EditorStyle/Private/EditorStylePrivatePCH.h
@@ -17,7 +17,7 @@
 #include "Settings.h"
 #include "Slate.h"
 #include "SlateStyle.h"
-#include "TargetPlatform.h"
+#include "PlatformInfo.h"
 
 
 /* Private includes
diff --git a/Engine/Source/Editor/EditorStyle/Private/SlateEditorStyle.cpp b/Engine/Source/Editor/EditorStyle/Private/SlateEditorStyle.cpp
--- a/Engine/Source/Editor/EditorStyle/Private/SlateEditorStyle.cpp
+++ b/Engine/Source/Editor/EditorStyle/Private/SlateEditorStyle.cpp
@@ -5033,25 +5033,15 @@ void FSlateEditorStyle::FStyle::SetupAutomationStyles()
 		Set( "Launcher.Platform.Warning", new IMAGE_BRUSH( "Icons/alert", Icon24x24) );
 
 #if (WITH_EDITOR || (IS_PROGRAM && PLATFORM_DESKTOP))
-		const TArray<ITargetPlatform*>& TargetPlatforms = GetTargetPlatformManager()->GetTargetPlatforms();
-		for (int32 PlatformIdx = 0; PlatformIdx < TargetPlatforms.Num(); PlatformIdx++)
+		Set( "Launcher.Platform.AllPlatforms", new IMAGE_BRUSH( "Launcher/All_Platforms_24x", Icon24x24) );
+		Set( "Launcher.Platform.AllPlatforms.Large", new IMAGE_BRUSH( "Launcher/All_Platforms_128x", Icon64x64) );
+		Set( "Launcher.Platform.AllPlatforms.XLarge", new IMAGE_BRUSH( "Launcher/All_Platforms_128x", Icon128x128) );
+
+		for(const PlatformInfo::FPlatformInfo& PlatformInfo : PlatformInfo::EnumeratePlatformInfoArray())
 		{
-			ITargetPlatform* TargetPlatform = TargetPlatforms[PlatformIdx];
-			if (TargetPlatform)
-			{
-				// Grab each icon name
-				FString PropertyName = FString::Printf(TEXT("Launcher.Platform_%s"), *(TargetPlatform->PlatformName()));
-				FString IconPath = TargetPlatform->GetIconPath(ETargetPlatformIcons::Normal);
-				Set( *PropertyName, new IMAGE_BRUSH( *IconPath, Icon24x24 ) );
-
-				PropertyName = FString::Printf(TEXT("Launcher.Platform_%s.Large"), *(TargetPlatform->PlatformName()));
-				IconPath = TargetPlatform->GetIconPath(ETargetPlatformIcons::Large);
-				Set( *PropertyName, new IMAGE_BRUSH( *IconPath, Icon64x64 ) );
-
-				PropertyName = FString::Printf(TEXT("Launcher.Platform_%s.XLarge"), *(TargetPlatform->PlatformName()));
-				IconPath = TargetPlatform->GetIconPath(ETargetPlatformIcons::XLarge);
-				Set( *PropertyName, new IMAGE_BRUSH( *IconPath, Icon128x128 ) );
-			}
+			Set( PlatformInfo.GetIconStyleName(PlatformInfo::EPlatformIconSize::Normal), new IMAGE_BRUSH( *PlatformInfo.GetIconPath(PlatformInfo::EPlatformIconSize::Normal), Icon24x24 ) );
+			Set( PlatformInfo.GetIconStyleName(PlatformInfo::EPlatformIconSize::Large),  new IMAGE_BRUSH( *PlatformInfo.GetIconPath(PlatformInfo::EPlatformIconSize::Large),  Icon64x64 ) );
+			Set( PlatformInfo.GetIconStyleName(PlatformInfo::EPlatformIconSize::XLarge), new IMAGE_BRUSH( *PlatformInfo.GetIconPath(PlatformInfo::EPlatformIconSize::XLarge), Icon128x128 ) );
 		}
 #endif
 	}
diff --git a/Engine/Source/Editor/GameProjectGeneration/Private/GameProjectGenerationModule.cpp b/Engine/Source/Editor/GameProjectGeneration/Private/GameProjectGenerationModule.cpp
--- a/Engine/Source/Editor/GameProjectGeneration/Private/GameProjectGenerationModule.cpp
+++ b/Engine/Source/Editor/GameProjectGeneration/Private/GameProjectGenerationModule.cpp
@@ -89,4 +89,9 @@ void FGameProjectGenerationModule::UpdateSupportedTargetPlatforms(const FName& I
 	GameProjectUtils::UpdateSupportedTargetPlatforms(InPlatformName, bIsSupported);
 }
 
+void FGameProjectGenerationModule::ClearSupportedTargetPlatforms()
+{
+	GameProjectUtils::ClearSupportedTargetPlatforms();
+}
+
 #undef LOCTEXT_NAMESPACE
diff --git a/Engine/Source/Editor/GameProjectGeneration/Private/GameProjectUtils.cpp b/Engine/Source/Editor/GameProjectGeneration/Private/GameProjectUtils.cpp
--- a/Engine/Source/Editor/GameProjectGeneration/Private/GameProjectUtils.cpp
+++ b/Engine/Source/Editor/GameProjectGeneration/Private/GameProjectUtils.cpp
@@ -2005,6 +2005,28 @@ void GameProjectUtils::UpdateSupportedTargetPlatforms(const FName& InPlatformNam
 	}
 }
 
+void GameProjectUtils::ClearSupportedTargetPlatforms()
+{
+	const FString& ProjectFilename = FPaths::GetProjectFilePath();
+	if(!ProjectFilename.IsEmpty())
+	{
+		// First attempt to check out the file if SCC is enabled
+		if(ISourceControlModule::Get().IsEnabled())
+		{
+			FText UnusedFailReason;
+			CheckoutGameProjectFile(ProjectFilename, UnusedFailReason);
+		}
+
+		// Second make sure the file is writable
+		if(FPlatformFileManager::Get().GetPlatformFile().IsReadOnly(*ProjectFilename))
+		{
+			FPlatformFileManager::Get().GetPlatformFile().SetReadOnly(*ProjectFilename, false);
+		}
+
+		IProjectManager::Get().ClearSupportedTargetPlatformsForCurrentProject();
+	}
+}
+
 bool GameProjectUtils::ReadTemplateFile(const FString& TemplateFileName, FString& OutFileContents, FText& OutFailReason)
 {
 	const FString FullFileName = FPaths::EngineContentDir() / TEXT("Editor") / TEXT("Templates") / TemplateFileName;
diff --git a/Engine/Source/Editor/GameProjectGeneration/Private/GameProjectUtils.h b/Engine/Source/Editor/GameProjectGeneration/Private/GameProjectUtils.h
--- a/Engine/Source/Editor/GameProjectGeneration/Private/GameProjectUtils.h
+++ b/Engine/Source/Editor/GameProjectGeneration/Private/GameProjectUtils.h
@@ -237,6 +237,9 @@ public:
 	 */
 	static void UpdateSupportedTargetPlatforms(const FName& InPlatformName, const bool bIsSupported);
 
+	/** Clear the list of supported target platforms */
+	static void ClearSupportedTargetPlatforms();
+
 private:
 	/** Generates a new project without using a template project */
 	static bool GenerateProjectFromScratch(const FString& NewProjectFile, bool bShouldGenerateCode, bool bCopyStarterContent, FText& OutFailReason);
diff --git a/Engine/Source/Editor/GameProjectGeneration/Private/SProjectBrowser.cpp b/Engine/Source/Editor/GameProjectGeneration/Private/SProjectBrowser.cpp
--- a/Engine/Source/Editor/GameProjectGeneration/Private/SProjectBrowser.cpp
+++ b/Engine/Source/Editor/GameProjectGeneration/Private/SProjectBrowser.cpp
@@ -8,6 +8,7 @@
 #include "UProjectInfo.h"
 #include "SourceCodeNavigation.h"
 #include "TargetPlatform.h"
+#include "PlatformInfo.h"
 
 #define LOCTEXT_NAMESPACE "ProjectBrowser"
 
@@ -25,8 +26,9 @@ struct FProjectItem
 	TSharedPtr<FSlateBrush> ProjectThumbnail;
 	bool bIsNewProjectItem;
 	TArray<FName> TargetPlatforms;
+	bool bSupportsAllPlatforms;
 
-	FProjectItem(const FText& InName, const FText& InDescription, const FString& InEngineIdentifier, bool InUpToDate, const TSharedPtr<FSlateBrush>& InProjectThumbnail, const FString& InProjectFile, bool InIsNewProjectItem, TArray<FName> InTargetPlatforms)
+	FProjectItem(const FText& InName, const FText& InDescription, const FString& InEngineIdentifier, bool InUpToDate, const TSharedPtr<FSlateBrush>& InProjectThumbnail, const FString& InProjectFile, bool InIsNewProjectItem, TArray<FName> InTargetPlatforms, bool InSupportsAllPlatforms)
 		: Name(InName)
 		, Description(InDescription)
 		, EngineIdentifier(InEngineIdentifier)
@@ -35,6 +37,7 @@ struct FProjectItem
 		, ProjectThumbnail(InProjectThumbnail)
 		, bIsNewProjectItem(InIsNewProjectItem)
 		, TargetPlatforms(InTargetPlatforms)
+		, bSupportsAllPlatforms(InSupportsAllPlatforms)
 	{ }
 
 	/** Check if this project is up to date */
@@ -300,7 +303,7 @@ void SProjectBrowser::ConstructCategory( const TSharedRef<SVerticalBox>& InCateg
 		.OnGenerateTile(this, &SProjectBrowser::MakeProjectViewWidget)
 		.OnContextMenuOpening(this, &SProjectBrowser::OnGetContextMenuContent)
 		.OnMouseButtonDoubleClick(this, &SProjectBrowser::HandleProjectItemDoubleClick)
-		.OnSelectionChanged(TSlateDelegates<TSharedPtr<FProjectItem>>::FOnSelectionChanged::CreateSP(this, &SProjectBrowser::HandleProjectViewSelectionChanged, Category->CategoryName))
+		.OnSelectionChanged(this, &SProjectBrowser::HandleProjectViewSelectionChanged, Category->CategoryName)
 		.ItemHeight(ThumbnailSize + ThumbnailBorderPadding + 32)
 		.ItemWidth(ThumbnailSize + ThumbnailBorderPadding)
 	];
@@ -457,6 +460,9 @@ TSharedRef<SToolTip> SProjectBrowser::MakeProjectToolTip( TSharedPtr<FProjectIte
 	TSharedRef<SHorizontalBox> TargetPlatformIconsBox = SNew(SHorizontalBox);
 	for(const FName& PlatformName : ProjectItem->TargetPlatforms)
 	{
+		const PlatformInfo::FPlatformInfo* const PlatformInfo = PlatformInfo::FindPlatformInfo(PlatformName);
+		check(PlatformInfo);
+
 		TargetPlatformIconsBox->AddSlot()
 		.AutoWidth()
 		.HAlign(HAlign_Center)
@@ -468,7 +474,7 @@ TSharedRef<SToolTip> SProjectBrowser::MakeProjectToolTip( TSharedPtr<FProjectIte
 			.HeightOverride(20)
 			[
 				SNew(SImage)
-				.Image(FEditorStyle::GetBrush(*FString::Printf(TEXT("Launcher.Platform_%s"), *PlatformName.ToString())))
+				.Image(FEditorStyle::GetBrush(PlatformInfo->GetIconStyleName(PlatformInfo::EPlatformIconSize::Normal)))
 			]
 		];
 	}
@@ -497,24 +503,17 @@ TSharedRef<SToolTip> SProjectBrowser::MakeProjectToolTip( TSharedPtr<FProjectIte
 					.AutoHeight()
 					.VAlign(VAlign_Center)
 					[
-						SNew(SHorizontalBox)
-
-						+SHorizontalBox::Slot()
-						.AutoWidth()
-						.HAlign(HAlign_Left)
-						.VAlign(VAlign_Center)
-						[
-							SNew(STextBlock)
-							.Text( ProjectItem->Name )
-							.Font( FEditorStyle::GetFontStyle("ProjectBrowser.TileViewTooltip.NameFont") )
-						]
+						SNew(STextBlock)
+						.Text( ProjectItem->Name )
+						.Font( FEditorStyle::GetFontStyle("ProjectBrowser.TileViewTooltip.NameFont") )
+					]
 
-						+SHorizontalBox::Slot()
-						.HAlign(HAlign_Right)
-						.VAlign(VAlign_Center)
-						[
-							TargetPlatformIconsBox
-						]
+					+SVerticalBox::Slot()
+					.AutoHeight()
+					.VAlign(VAlign_Center)
+					.Padding(0, 2, 0, 0)
+					[
+						TargetPlatformIconsBox
 					]
 				]
 			]
@@ -705,7 +704,6 @@ FReply SProjectBrowser::FindProjects()
 
 	// Add all the discovered projects to the list
 	const FString EngineIdentifier = FDesktopPlatformModule::Get()->GetCurrentEngineIdentifier();
-	TArray<ITargetPlatform*> AllPlatforms = GetTargetPlatformManager()->GetTargetPlatforms();
 	for(TMap<FString, FText>::TConstIterator Iter(AbsoluteProjectFilesToCategory); Iter; ++Iter)
 	{
 		const FString& ProjectFilename = *Iter.Key();
@@ -751,24 +749,22 @@ FReply SProjectBrowser::FindProjects()
 				FString ProjectEngineIdentifier;
 				bool bIsUpToDate = FDesktopPlatformModule::Get()->GetEngineIdentifierForProject(ProjectFilename, ProjectEngineIdentifier) && ProjectEngineIdentifier == EngineIdentifier;
 
-				// Work out which platforms this project is targeting (it only shows the ones we have support for locally)
+				// Work out which platforms this project is targeting
 				TArray<FName> TargetPlatforms;
-				for ( ITargetPlatform* Platform : AllPlatforms )
+				for(const PlatformInfo::FPlatformInfo& PlatformInfo : PlatformInfo::EnumeratePlatformInfoArray())
 				{
-					// We are only interested in standalone games
-					const bool bIsGamePlatform = !Platform->IsClientOnly() && !Platform->IsServerOnly() && !Platform->HasEditorOnlyData();
-					if ( bIsGamePlatform )
+					if(PlatformInfo.IsVanilla() && PlatformInfo.PlatformType == PlatformInfo::EPlatformType::Game && ProjectStatus.IsTargetPlatformSupported(PlatformInfo.PlatformInfoName))
 					{
-						const FName PlatformName = *Platform->PlatformName();
-						if ( ProjectStatus.IsTargetPlatformSupported(PlatformName) )
-						{
-							TargetPlatforms.Add(PlatformName);
-						}
+						TargetPlatforms.Add(PlatformInfo.PlatformInfoName);
 					}
 				}
+				TargetPlatforms.Sort([](const FName& One, const FName& Two) -> bool
+				{
+					return One < Two;
+				});
 
 				const bool bIsNewProjectItem = false;
-				TSharedRef<FProjectItem> NewProjectItem = MakeShareable( new FProjectItem(ProjectName, ProjectDescription, ProjectEngineIdentifier, bIsUpToDate, DynamicBrush, ProjectFilename, bIsNewProjectItem, TargetPlatforms ) );
+				TSharedRef<FProjectItem> NewProjectItem = MakeShareable( new FProjectItem(ProjectName, ProjectDescription, ProjectEngineIdentifier, bIsUpToDate, DynamicBrush, ProjectFilename, bIsNewProjectItem, TargetPlatforms, ProjectStatus.SupportsAllPlatforms() ) );
 				AddProjectToCategory(NewProjectItem, ProjectCategory);
 			}
 		}
diff --git a/Engine/Source/Editor/GameProjectGeneration/Public/GameProjectGenerationModule.h b/Engine/Source/Editor/GameProjectGeneration/Public/GameProjectGenerationModule.h
--- a/Engine/Source/Editor/GameProjectGeneration/Public/GameProjectGenerationModule.h
+++ b/Engine/Source/Editor/GameProjectGeneration/Public/GameProjectGenerationModule.h
@@ -74,6 +74,9 @@ public:
 	 */
 	virtual void UpdateSupportedTargetPlatforms(const FName& InPlatformName, const bool bIsSupported);
 
+	/** Clear the list of supported target platforms */
+	virtual void ClearSupportedTargetPlatforms();
+
 private:
 	FAddCodeToProjectDialogOpenedEvent AddCodeToProjectDialogOpenedEvent;
 
diff --git a/Engine/Source/Editor/LevelEditor/Private/LevelEditor.cpp b/Engine/Source/Editor/LevelEditor/Private/LevelEditor.cpp
--- a/Engine/Source/Editor/LevelEditor/Private/LevelEditor.cpp
+++ b/Engine/Source/Editor/LevelEditor/Private/LevelEditor.cpp
@@ -86,11 +86,6 @@ TSharedRef<SDockTab> FLevelEditorModule::SpawnLevelEditor( const FSpawnTabArgs&
 			RightContentText = FText::Format(NSLOCTEXT("UnrealEditor", "TitleBarRightContent", "{Branch}{GameName}"), Args);
 		}
 
-		// Create the target platform icons, and register our interest so we can update them if they're changed later on
-		TargetPlatformIconsBox = SNew(SHorizontalBox);
-		UpdateTargetPlatformIcons();
-		IProjectManager::Get().OnTargetPlatformsForCurrentProjectChanged().AddRaw(this, &FLevelEditorModule::UpdateTargetPlatformIcons);
-
 		RightContent =
 				SNew( SHorizontalBox )
 
@@ -107,14 +102,6 @@ TSharedRef<SDockTab> FLevelEditorModule::SpawnLevelEditor( const FSpawnTabArgs&
 					]
 				]
 
-				+SHorizontalBox::Slot()
-				.AutoWidth()
-				.Padding(16.0f, 0.0f, 0.0f, 0.0f)
-				.VAlign(VAlign_Center)
-				[
-					TargetPlatformIconsBox.ToSharedRef()
-				]
-
 				+SHorizontalBox::Slot()
 				.AutoWidth()
 				.Padding(16.0f, 0.0f, 0.0f, 0.0f)
@@ -491,49 +478,6 @@ TSharedPtr<SDockTab> FLevelEditorModule::GetLevelEditorTab() const
 	return LevelEditorInstanceTabPtr.Pin().ToSharedRef();
 }
 
-void FLevelEditorModule::UpdateTargetPlatformIcons()
-{
-	TargetPlatformIconsBox->ClearChildren();
-
-	FProjectStatus ProjectStatus;
-	if(IProjectManager::Get().QueryStatusForCurrentProject(ProjectStatus))
-	{
-		TArray<ITargetPlatform*> AllPlatforms = GetTargetPlatformManager()->GetTargetPlatforms();
-		TArray<FName> TargetPlatforms;
-		for( ITargetPlatform* Platform : AllPlatforms)
-		{
-			// We are only interested in standalone games
-			const bool bIsGamePlatform = !Platform->IsClientOnly() && !Platform->IsServerOnly() && !Platform->HasEditorOnlyData();
-			if(bIsGamePlatform)
-			{
-				const FName PlatformName = *Platform->PlatformName();
-				if(ProjectStatus.IsTargetPlatformSupported(PlatformName))
-				{
-					TargetPlatforms.Add(PlatformName);
-				}
-			}
-		}
-
-		for(const FName& PlatformName : TargetPlatforms)
-		{
-			TargetPlatformIconsBox->AddSlot()
-			.AutoWidth()
-			.HAlign(HAlign_Center)
-			.VAlign(VAlign_Center)
-			.Padding(FMargin(2, 0, 2, 0))
-			[
-				SNew(SBox)
-				.WidthOverride(20)
-				.HeightOverride(20)
-				[
-					SNew(SImage)
-					.Image(FEditorStyle::GetBrush(*FString::Printf(TEXT("Launcher.Platform_%s"), *PlatformName.ToString())))
-				]
-			];
-		}
-	}
-}
-
 void FLevelEditorModule::BindGlobalLevelEditorCommands()
 {
 	check( !GlobalLevelEditorActions.IsValid() );
diff --git a/Engine/Source/Editor/LevelEditor/Public/LevelEditor.h b/Engine/Source/Editor/LevelEditor/Public/LevelEditor.h
--- a/Engine/Source/Editor/LevelEditor/Public/LevelEditor.h
+++ b/Engine/Source/Editor/LevelEditor/Public/LevelEditor.h
@@ -237,9 +237,6 @@ private:
 	/** Spawn the main level editor tab */
 	TSharedRef<SDockTab> SpawnLevelEditor( const FSpawnTabArgs& InArgs );
 
-	/** Update the target platform icons shown in the titlebar */
-	void UpdateTargetPlatformIcons();
-
 private:
 	TSharedPtr<FExtensibilityManager> MenuExtensibilityManager;
 	TSharedPtr<FExtensibilityManager> ToolBarExtensibilityManager;
@@ -286,9 +283,6 @@ private:
 	/* Holds the Editor's tab manager */
 	TSharedPtr<FTabManager> LevelEditorTabManager;
 
-	/* Container for the target platform icons for this project; updated whenever the project is updated via IProjectManager::OnTargetPlatformsForCurrentProjectChanged */
-	TSharedPtr<SHorizontalBox> TargetPlatformIconsBox;
-
 	/* Whether we have source for the level editor, and it can be recompiled. */
 	bool bCanBeRecompiled;
 };
diff --git a/Engine/Source/Editor/MainFrame/MainFrame.Build.cs b/Engine/Source/Editor/MainFrame/MainFrame.Build.cs
--- a/Engine/Source/Editor/MainFrame/MainFrame.Build.cs
+++ b/Engine/Source/Editor/MainFrame/MainFrame.Build.cs
@@ -29,6 +29,7 @@ public class MainFrame : ModuleRules
 				"SourceControl",
 				"SourceControlWindows",
 				"TargetPlatform",
+				"DesktopPlatform",
 				"UnrealEd",
 				"WorkspaceMenuStructure",
 				"MessageLog",
diff --git a/Engine/Source/Editor/MainFrame/Private/Frame/MainFrameActions.cpp b/Engine/Source/Editor/MainFrame/Private/Frame/MainFrameActions.cpp
--- a/Engine/Source/Editor/MainFrame/Private/Frame/MainFrameActions.cpp
+++ b/Engine/Source/Editor/MainFrame/Private/Frame/MainFrameActions.cpp
@@ -398,31 +398,36 @@ void FMainFrameActionCallbacks::AddCodeToProject()
 	FGameProjectGenerationModule::Get().OpenAddCodeToProjectDialog();
 }
 
-void FMainFrameActionCallbacks::CookContent( const FString InPlatformName, const FText PlatformDisplayName )
+void FMainFrameActionCallbacks::CookContent( const FName InPlatformInfoName )
 {
-	FString PlatformName = InPlatformName;
+	const PlatformInfo::FPlatformInfo* const PlatformInfo = PlatformInfo::FindPlatformInfo(InPlatformInfoName);
+	check(PlatformInfo);
+
 	FString OptionalParams;
 
-	if (!FModuleManager::LoadModuleChecked<IProjectTargetPlatformEditorModule>("ProjectTargetPlatformEditor").ShowUnsupportedTargetWarning(*PlatformName))
+	if (!FModuleManager::LoadModuleChecked<IProjectTargetPlatformEditorModule>("ProjectTargetPlatformEditor").ShowUnsupportedTargetWarning(PlatformInfo->VanillaPlatformName))
 	{
 		return;
 	}
 
-	if (PlatformName == "MacNoEditor")
+	if (PlatformInfo->TargetPlatformName == FName("MacNoEditor"))
 	{
 		OptionalParams += TEXT(" -targetplatform=Mac");
 	}
+	else if (PlatformInfo->TargetPlatformName == FName("LinuxNoEditor"))
+	{
+		OptionalParams += TEXT(" -targetplatform=Linux");
+	}
 
-	FString ExecutableName = FPlatformProcess::ExecutableName(false);
-
-	// for things like Android_ATC, pull it apart to be Android, with cook flavor ATC
-	int32 UnderscoreLoc;
-	if (PlatformName.FindChar(TEXT('_'), UnderscoreLoc))
+	// Append any extra UAT flags specified for this platform flavor
+	if (!PlatformInfo->UATCommandLine.IsEmpty())
 	{
-		OptionalParams += FString(TEXT(" -cookflavor=")) + PlatformName.Mid(UnderscoreLoc + 1);
-		PlatformName = PlatformName.Mid(0, UnderscoreLoc);
+		OptionalParams += TEXT(" ");
+		OptionalParams += PlatformInfo->UATCommandLine;
 	}
 
+	FString ExecutableName = FPlatformProcess::ExecutableName(false);
+
 #if PLATFORM_WINDOWS
 	// turn UE4editor into UE4editor-cmd
 	if(ExecutableName.EndsWith(".exe", ESearchCase::IgnoreCase) && !FPaths::GetBaseFilename(ExecutableName).EndsWith("-cmd", ESearchCase::IgnoreCase))
@@ -446,15 +451,15 @@ void FMainFrameActionCallbacks::CookContent( const FString InPlatformName, const
 		FRocketSupport::IsRocket() ? TEXT("-rocket -nocompile") : TEXT("-nocompileeditor"),
 		FApp::IsEngineInstalled() ? TEXT(" -installed") : TEXT(""),
 		*ProjectPath,
-		*PlatformName,
+		*PlatformInfo->TargetPlatformName.ToString(),
 		*ExecutableName,
 		*OptionalParams
 	);
 
-	CreateUatTask(CommandLine, PlatformDisplayName, LOCTEXT("CookingContentTaskName", "Cooking content"), LOCTEXT("CookingTaskName", "Cooking"), FEditorStyle::GetBrush(TEXT("MainFrame.CookContent")));
+	CreateUatTask(CommandLine, PlatformInfo->DisplayName, LOCTEXT("CookingContentTaskName", "Cooking content"), LOCTEXT("CookingTaskName", "Cooking"), FEditorStyle::GetBrush(TEXT("MainFrame.CookContent")));
 }
 
-bool FMainFrameActionCallbacks::CookContentCanExecute( const FString PlatformName )
+bool FMainFrameActionCallbacks::CookContentCanExecute( const FName PlatformInfoName )
 {
 	return true;
 }
@@ -470,36 +475,39 @@ bool FMainFrameActionCallbacks::PackageBuildConfigurationIsChecked( EProjectPack
 	return (GetDefault<UProjectPackagingSettings>()->BuildConfiguration == BuildConfiguration);
 }
 
-void FMainFrameActionCallbacks::PackageProject( const FString InPlatformName, const FText PlatformDisplayName, const FString AdditionalFlags )
+void FMainFrameActionCallbacks::PackageProject( const FName InPlatformInfoName )
 {
-	FString PlatformName = InPlatformName;
-
 	// does the project have any code?
 	FGameProjectGenerationModule& GameProjectModule = FModuleManager::LoadModuleChecked<FGameProjectGenerationModule>(TEXT("GameProjectGeneration"));
 	bool bProjectHasCode = GameProjectModule.Get().GetProjectCodeFileCount() > 0;
 
-	const ITargetPlatform* Platform = GetTargetPlatformManager()->FindTargetPlatform(PlatformName);
-	if (Platform)
+	const PlatformInfo::FPlatformInfo* const PlatformInfo = PlatformInfo::FindPlatformInfo(InPlatformInfoName);
+	check(PlatformInfo);
+
 	{
-		FString NotInstalledDocLink;
-		if (!Platform->IsSdkInstalled(bProjectHasCode, NotInstalledDocLink))
+		const ITargetPlatform* const Platform = GetTargetPlatformManager()->FindTargetPlatform(PlatformInfo->TargetPlatformName.ToString());
+		if (Platform)
 		{
-			// broadcast this, and assume someone will pick it up
-			IMainFrameModule& MainFrameModule = FModuleManager::GetModuleChecked<IMainFrameModule>(TEXT("MainFrame"));
-			MainFrameModule.BroadcastMainFrameSDKNotInstalled(PlatformName, NotInstalledDocLink);
-			return;
+			FString NotInstalledDocLink;
+			if (!Platform->IsSdkInstalled(bProjectHasCode, NotInstalledDocLink))
+			{
+				// broadcast this, and assume someone will pick it up
+				IMainFrameModule& MainFrameModule = FModuleManager::GetModuleChecked<IMainFrameModule>(TEXT("MainFrame"));
+				MainFrameModule.BroadcastMainFrameSDKNotInstalled(PlatformInfo->TargetPlatformName.ToString(), NotInstalledDocLink);
+				return;
+			}
 		}
 	}
 
 #if PLATFORM_WINDOWS
-	if (bProjectHasCode && FRocketSupport::IsRocket() && PlatformName == TEXT("IOS"))
+	if (bProjectHasCode && FRocketSupport::IsRocket() && PlatformInfo->TargetPlatformName == FName("IOS"))
 	{
 		FMessageDialog::Open(EAppMsgType::Ok, LOCTEXT("PackagingNotWorkingMessage", "Sorry, packaging is currently not supported for code-based iOS projects. This feature will be available in a future release.") );
 		return;
 	}
 #endif
 
-	if (!FModuleManager::LoadModuleChecked<IProjectTargetPlatformEditorModule>("ProjectTargetPlatformEditor").ShowUnsupportedTargetWarning(*PlatformName))
+	if (!FModuleManager::LoadModuleChecked<IProjectTargetPlatformEditorModule>("ProjectTargetPlatformEditor").ShowUnsupportedTargetWarning(PlatformInfo->VanillaPlatformName))
 	{
 		return;
 	}
@@ -553,34 +561,26 @@ void FMainFrameActionCallbacks::PackageProject( const FString InPlatformName, co
 		OptionalParams += TEXT(" -distribution");
 	}
 
-	if (PlatformName == "WindowsNoEditor")
-	{
-		FString TargetPlatform;
-		FParse::Value(*AdditionalFlags, TEXT("-targetplatform="), TargetPlatform);
-
-		OptionalParams += FString::Printf(TEXT(" -targetplatform=%s"), *TargetPlatform);
-	}
-	else if (PlatformName == "MacNoEditor")
+	if (PlatformInfo->TargetPlatformName == FName("MacNoEditor"))
 	{
 		OptionalParams += TEXT(" -targetplatform=Mac");
 	}
-	else if (PlatformName == "LinuxNoEditor")
+	else if (PlatformInfo->TargetPlatformName == FName("LinuxNoEditor"))
 	{
 		OptionalParams += TEXT(" -targetplatform=Linux");
 	}
 
-	// only build if the project has code that might need to be built
-	if (bProjectHasCode && FSourceCodeNavigation::IsCompilerAvailable())
+	// Append any extra UAT flags specified for this platform flavor
+	if (!PlatformInfo->UATCommandLine.IsEmpty())
 	{
-		OptionalParams += TEXT(" -build");
+		OptionalParams += TEXT(" ");
+		OptionalParams += PlatformInfo->UATCommandLine;
 	}
 
-	// for things like Android_ATC, pull it apart to be Android, with cook flavor ATC
-	int32 UnderscoreLoc;
-	if (PlatformName.FindChar(TEXT('_'), UnderscoreLoc))
+	// only build if the project has code that might need to be built
+	if (bProjectHasCode && FSourceCodeNavigation::IsCompilerAvailable())
 	{
-		OptionalParams += FString(TEXT(" -cookflavor=")) + PlatformName.Mid(UnderscoreLoc + 1);
-		PlatformName = PlatformName.Mid(0, UnderscoreLoc);
+		OptionalParams += TEXT(" -build");
 	}
 
 	FString ExecutableName = FPlatformProcess::ExecutableName(false);
@@ -612,16 +612,16 @@ void FMainFrameActionCallbacks::PackageProject( const FString InPlatformName, co
 		FApp::IsEngineInstalled() ? TEXT(" -installed") : TEXT(""),
 		*ProjectPath,
 		*PackagingSettings->StagingDirectory.Path,
-		*PlatformName,
+		*PlatformInfo->TargetPlatformName.ToString(),
 		*Configuration,
 		*ExecutableName,
 		*OptionalParams
 	);
 
-	CreateUatTask(CommandLine, PlatformDisplayName, LOCTEXT("PackagingProjectTaskName", "Packaging project"), LOCTEXT("PackagingTaskName", "Packaging"), FEditorStyle::GetBrush(TEXT("MainFrame.PackageProject")));
+	CreateUatTask(CommandLine, PlatformInfo->DisplayName, LOCTEXT("PackagingProjectTaskName", "Packaging project"), LOCTEXT("PackagingTaskName", "Packaging"), FEditorStyle::GetBrush(TEXT("MainFrame.PackageProject")));
 }
 
-bool FMainFrameActionCallbacks::PackageProjectCanExecute( const FString PlatformName, bool IsImplemented )
+bool FMainFrameActionCallbacks::PackageProjectCanExecute( const FName PlatformInfoName, bool IsImplemented )
 {
 	return IsImplemented;
 }
diff --git a/Engine/Source/Editor/MainFrame/Private/Frame/MainFrameActions.h b/Engine/Source/Editor/MainFrame/Private/Frame/MainFrameActions.h
--- a/Engine/Source/Editor/MainFrame/Private/Frame/MainFrameActions.h
+++ b/Engine/Source/Editor/MainFrame/Private/Frame/MainFrameActions.h
@@ -149,10 +149,10 @@ public:
 	static void AddCodeToProject();
 
 	/** Cooks the project's content for the specified platform. */
-	static void CookContent( const FString PlatformName, const FText PlatformDisplayName );
+	static void CookContent( const FName InPlatformInfoName );
 
 	/** Checks whether a menu action for cooking the project's content can execute. */
-	static bool CookContentCanExecute( const FString PlatformName );
+	static bool CookContentCanExecute( const FName PlatformInfoName );
 
 	/** Sets the project packaging build configuration. */
 	static void PackageBuildConfiguration( EProjectPackagingBuildConfigurations BuildConfiguration );
@@ -161,10 +161,10 @@ public:
 	static bool PackageBuildConfigurationIsChecked( EProjectPackagingBuildConfigurations BuildConfiguration );
 
 	/** Packages the project for the specified platform. */
-	static void PackageProject( const FString PlatformName, const FText PlatformDisplayName, const FString AdditionalFlags );
+	static void PackageProject( const FName InPlatformInfoName );
 
 	/** Checks whether a menu action for packaging the project can execute. */
-	static bool PackageProjectCanExecute( const FString PlatformName, bool IsImplemented );
+	static bool PackageProjectCanExecute( const FName PlatformInfoName, bool IsImplemented );
 
 	/** Refresh the project in the current IDE */
 	static void RefreshCodeProject();
diff --git a/Engine/Source/Editor/MainFrame/Private/MainFramePrivatePCH.h b/Engine/Source/Editor/MainFrame/Private/MainFramePrivatePCH.h
--- a/Engine/Source/Editor/MainFrame/Private/MainFramePrivatePCH.h
+++ b/Engine/Source/Editor/MainFrame/Private/MainFramePrivatePCH.h
@@ -33,6 +33,7 @@
 #include "SourceControlWindows.h"
 #include "Settings.h"
 #include "TargetPlatform.h"
+#include "PlatformInfo.h"
 #include "UnrealEd.h"
 
 #include "../../LevelEditor/Public/LevelEditor.h"
diff --git a/Engine/Source/Editor/MainFrame/Private/Menus/CookContentMenu.h b/Engine/Source/Editor/MainFrame/Private/Menus/CookContentMenu.h
--- a/Engine/Source/Editor/MainFrame/Private/Menus/CookContentMenu.h
+++ b/Engine/Source/Editor/MainFrame/Private/Menus/CookContentMenu.h
@@ -6,6 +6,7 @@
 
 #pragma once
 
+#include "ProjectTargetPlatformEditor.h"
 
 #define LOCTEXT_NAMESPACE "FCookContentMenu"
 
@@ -24,34 +25,103 @@ public:
 	 */
 	static void MakeMenu( FMenuBuilder& MenuBuilder )
 	{
-		TArray<ITargetPlatform*> Platforms = GetTargetPlatformManager()->GetTargetPlatforms();
-
-		if (Platforms.Num() == 0)
+		TArray<PlatformInfo::FVanillaPlatformEntry> VanillaPlatforms = PlatformInfo::BuildPlatformHierarchy(PlatformInfo::EPlatformFilter::CookFlavor);
+		if (!VanillaPlatforms.Num())
 		{
 			return;
 		}
 
-		for (int32 PlatformIndex = 0; PlatformIndex < Platforms.Num(); ++PlatformIndex)
+		VanillaPlatforms.Sort([](const PlatformInfo::FVanillaPlatformEntry& One, const PlatformInfo::FVanillaPlatformEntry& Two) -> bool
+		{
+			return One.PlatformInfo->DisplayName.CompareTo(Two.PlatformInfo->DisplayName) < 0;
+		});
+
+		IProjectTargetPlatformEditorModule& ProjectTargetPlatformEditorModule = FModuleManager::LoadModuleChecked<IProjectTargetPlatformEditorModule>("ProjectTargetPlatformEditor");
+
+		// Build up a menu from the tree of platforms
+		for (const PlatformInfo::FVanillaPlatformEntry& VanillaPlatform : VanillaPlatforms)
 		{
-			const ITargetPlatform* Platform = Platforms[PlatformIndex];
+			check(VanillaPlatform.PlatformInfo->IsVanilla());
 
-			// for the Editor we are only interested in cooking standalone games
-			if (Platform->IsClientOnly() || Platform->IsServerOnly() || Platform->HasEditorOnlyData())
+			// Only care about game targets
+			if (VanillaPlatform.PlatformInfo->PlatformType != PlatformInfo::EPlatformType::Game)
 			{
 				continue;
 			}
 
-			FUIAction Action(
-				FExecuteAction::CreateStatic(&FMainFrameActionCallbacks::CookContent, Platform->PlatformName(), Platform->DisplayName()),
-				FCanExecuteAction::CreateStatic(&FMainFrameActionCallbacks::CookContentCanExecute, Platform->PlatformName())
-			);
+			// Make sure we're able to run this platform
+			ITargetPlatform* const Platform = GetTargetPlatformManager()->FindTargetPlatform(VanillaPlatform.PlatformInfo->TargetPlatformName.ToString());
+			if (!Platform)
+			{
+				continue;
+			}
+
+			if(VanillaPlatform.PlatformFlavors.Num())
+			{
+				MenuBuilder.AddSubMenu(
+					ProjectTargetPlatformEditorModule.MakePlatformMenuItemWidget(*VanillaPlatform.PlatformInfo), 
+					FNewMenuDelegate::CreateStatic(&FCookContentMenu::AddPlatformSubPlatformsToMenu, VanillaPlatform.PlatformFlavors),
+					false
+					);
+			}
+			else
+			{
+				AddPlatformToMenu(MenuBuilder, *VanillaPlatform.PlatformInfo);
+			}
+		}
+	}
 
-			MenuBuilder.AddMenuEntry(
-				Platform->DisplayName(),
-				FText::Format(LOCTEXT("CookContentForPlatformTooltip", "Cook your game content for the {0} platform"), Platform->DisplayName()),
-				FSlateIcon(FEditorStyle::GetStyleSetName(), *FString::Printf(TEXT("Launcher.Platform_%s"), *Platform->PlatformName())),
-				Action
+protected:
+
+	/**
+	* Creates the platform menu entries.
+	*
+	* @param MenuBuilder	- The builder for the menu that owns this menu.
+	* @param Platform		- The target platform we allow packaging for
+	*/
+	static void AddPlatformToMenu(FMenuBuilder& MenuBuilder, const PlatformInfo::FPlatformInfo& PlatformInfo)
+	{
+		IProjectTargetPlatformEditorModule& ProjectTargetPlatformEditorModule = FModuleManager::LoadModuleChecked<IProjectTargetPlatformEditorModule>("ProjectTargetPlatformEditor");
+
+		FUIAction Action(
+			FExecuteAction::CreateStatic(&FMainFrameActionCallbacks::CookContent, PlatformInfo.PlatformInfoName),
+			FCanExecuteAction::CreateStatic(&FMainFrameActionCallbacks::CookContentCanExecute, PlatformInfo.PlatformInfoName)
+		);
+
+		// ... generate tooltip text
+		FFormatNamedArguments TooltipArguments;
+		TooltipArguments.Add(TEXT("DisplayName"), PlatformInfo.DisplayName);
+		FText Tooltip = FText::Format(LOCTEXT("CookContentForPlatformTooltip", "Cook your game content for the {DisplayName} platform"), TooltipArguments);
+
+		FProjectStatus ProjectStatus;
+		if (IProjectManager::Get().QueryStatusForCurrentProject(ProjectStatus) && !ProjectStatus.IsTargetPlatformSupported(PlatformInfo.VanillaPlatformName))
+		{
+			FText TooltipLine2 = FText::Format(LOCTEXT("CookUnsupportedPlatformWarning", "{DisplayName} is not listed as a target platform for this project, so may not run as expected."), TooltipArguments);
+			Tooltip = FText::Format(FText::FromString(TEXT("{0}\n\n{1}")), Tooltip, TooltipLine2);
+		}
+
+		// ... and add a menu entry
+		MenuBuilder.AddMenuEntry(
+			Action, 
+			ProjectTargetPlatformEditorModule.MakePlatformMenuItemWidget(PlatformInfo), 
+			NAME_None, 
+			Tooltip
 			);
+	}
+
+
+	/**
+	* Creates the platform menu entries for a given platforms sub-platforms.
+	* e.g. Windows has multiple sub-platforms - Win32 and Win64
+	*
+	* @param MenuBuilder		- The builder for the menu that owns this menu.
+	* @param SubPlatformInfos	- The Sub-platform information
+	*/
+	static void AddPlatformSubPlatformsToMenu(FMenuBuilder& MenuBuilder, TArray<const PlatformInfo::FPlatformInfo*> SubPlatformInfos)
+	{
+		for (const PlatformInfo::FPlatformInfo* SubPlatformInfo : SubPlatformInfos)
+		{
+			AddPlatformToMenu(MenuBuilder, *SubPlatformInfo);
 		}
 	}
 };
diff --git a/Engine/Source/Editor/MainFrame/Private/Menus/MainMenu.cpp b/Engine/Source/Editor/MainFrame/Private/Menus/MainMenu.cpp
--- a/Engine/Source/Editor/MainFrame/Private/Menus/MainMenu.cpp
+++ b/Engine/Source/Editor/MainFrame/Private/Menus/MainMenu.cpp
@@ -317,6 +317,14 @@ TSharedRef< SWidget > FMainMenu::MakeMainTabMenu( const TSharedPtr<FTabManager>&
 					FNewMenuDelegate::CreateStatic( &FPackageProjectMenu::MakeMenu ), false, FSlateIcon(FEditorStyle::GetStyleSetName(), "MainFrame.PackageProject")
 				);
 
+				/*
+				MenuBuilder.AddSubMenu(
+					LOCTEXT("CookProjectSubMenuLabel", "Cook Project"),
+					LOCTEXT("CookProjectSubMenuToolTip", "Cook your project content for debugging"),
+					FNewMenuDelegate::CreateStatic( &FCookContentMenu::MakeMenu ), false, FSlateIcon()
+				);
+				*/
+
 				if (FPaths::FileExists(FModuleManager::Get().GetSolutionFilepath()))
 				{
 					MenuBuilder.AddMenuEntry( FMainFrameCommands::Get().RefreshCodeProject,
@@ -337,11 +345,11 @@ TSharedRef< SWidget > FMainMenu::MakeMainTabMenu( const TSharedPtr<FTabManager>&
 
 				if (RunningTargetPlatform != nullptr)
 				{
-					FString CookedPlatformName = RunningTargetPlatform->PlatformName() + TEXT("NoEditor");
-					FText CookedPlatformText = FText::FromString(RunningTargetPlatform->PlatformName());
+					const FName CookedPlatformName = *(RunningTargetPlatform->PlatformName() + TEXT("NoEditor"));
+					const FText CookedPlatformText = FText::FromString(RunningTargetPlatform->PlatformName());
 
 					FUIAction Action(
-						FExecuteAction::CreateStatic(&FMainFrameActionCallbacks::CookContent, CookedPlatformName, CookedPlatformText),
+						FExecuteAction::CreateStatic(&FMainFrameActionCallbacks::CookContent, CookedPlatformName),
 						FCanExecuteAction::CreateStatic(&FMainFrameActionCallbacks::CookContentCanExecute, CookedPlatformName)
 					);
 
diff --git a/Engine/Source/Editor/MainFrame/Private/Menus/PackageProjectMenu.h b/Engine/Source/Editor/MainFrame/Private/Menus/PackageProjectMenu.h
--- a/Engine/Source/Editor/MainFrame/Private/Menus/PackageProjectMenu.h
+++ b/Engine/Source/Editor/MainFrame/Private/Menus/PackageProjectMenu.h
@@ -25,40 +25,48 @@ public:
 	 */
 	static void MakeMenu( FMenuBuilder& MenuBuilder )
 	{
-		TArray<ITargetPlatform*> Platforms = GetTargetPlatformManager()->GetTargetPlatforms();
-
-		if (Platforms.Num() == 0)
+		TArray<PlatformInfo::FVanillaPlatformEntry> VanillaPlatforms = PlatformInfo::BuildPlatformHierarchy(PlatformInfo::EPlatformFilter::All);
+		if (!VanillaPlatforms.Num())
 		{
 			return;
 		}
-		
+
+		VanillaPlatforms.Sort([](const PlatformInfo::FVanillaPlatformEntry& One, const PlatformInfo::FVanillaPlatformEntry& Two) -> bool
+		{
+			return One.PlatformInfo->DisplayName.CompareTo(Two.PlatformInfo->DisplayName) < 0;
+		});
+
 		IProjectTargetPlatformEditorModule& ProjectTargetPlatformEditorModule = FModuleManager::LoadModuleChecked<IProjectTargetPlatformEditorModule>("ProjectTargetPlatformEditor");
 
-		for (int32 PlatformIndex = 0; PlatformIndex < Platforms.Num(); ++PlatformIndex)
+		// Build up a menu from the tree of platforms
+		for (const PlatformInfo::FVanillaPlatformEntry& VanillaPlatform : VanillaPlatforms)
 		{
-			const ITargetPlatform* Platform = Platforms[PlatformIndex];
+			check(VanillaPlatform.PlatformInfo->IsVanilla());
 
-			// for the Editor we are only interested in packaging standalone games
-			if (Platform->IsClientOnly() || Platform->IsServerOnly() || Platform->HasEditorOnlyData())
+			// Only care about game targets
+			if (VanillaPlatform.PlatformInfo->PlatformType != PlatformInfo::EPlatformType::Game)
 			{
 				continue;
 			}
 
-			if (Platform->PlatformName() == TEXT("WindowsNoEditor"))
+			// Make sure we're able to run this platform
+			ITargetPlatform* const Platform = GetTargetPlatformManager()->FindTargetPlatform(VanillaPlatform.PlatformInfo->TargetPlatformName.ToString());
+			if (!Platform)
 			{
-				TArray< FText > CollectivePlatforms;
-				CollectivePlatforms.Add(LOCTEXT("Win32","Win32"));
-				CollectivePlatforms.Add(LOCTEXT("Win64","Win64"));
+				continue;
+			}
 
+			if (VanillaPlatform.PlatformFlavors.Num())
+			{
 				MenuBuilder.AddSubMenu(
-					ProjectTargetPlatformEditorModule.MakePlatformMenuItemWidget(Platform), 
-					FNewMenuDelegate::CreateStatic(&FPackageProjectMenu::AddPlatformSubPlatformsToMenu, Platform, CollectivePlatforms),
+					ProjectTargetPlatformEditorModule.MakePlatformMenuItemWidget(*VanillaPlatform.PlatformInfo), 
+					FNewMenuDelegate::CreateStatic(&FPackageProjectMenu::AddPlatformSubPlatformsToMenu, VanillaPlatform.PlatformFlavors),
 					false
 					);
 			}
 			else
 			{
-				AddPlatformToMenu(MenuBuilder, Platform);
+				AddPlatformToMenu(MenuBuilder, *VanillaPlatform.PlatformInfo);
 			}
 		}
 
@@ -83,31 +91,33 @@ protected:
 	* @param MenuBuilder	- The builder for the menu that owns this menu.
 	* @param Platform		- The target platform we allow packaging for
 	*/
-	static void AddPlatformToMenu(FMenuBuilder& MenuBuilder, const ITargetPlatform* Platform)
+	static void AddPlatformToMenu(FMenuBuilder& MenuBuilder, const PlatformInfo::FPlatformInfo& PlatformInfo)
 	{
 		IProjectTargetPlatformEditorModule& ProjectTargetPlatformEditorModule = FModuleManager::LoadModuleChecked<IProjectTargetPlatformEditorModule>("ProjectTargetPlatformEditor");
 
+		ITargetPlatform* const Platform = GetTargetPlatformManager()->FindTargetPlatform(PlatformInfo.TargetPlatformName.ToString());
+
 		FUIAction Action(
-			FExecuteAction::CreateStatic(&FMainFrameActionCallbacks::PackageProject, Platform->PlatformName(), Platform->DisplayName(), FString()),
-			FCanExecuteAction::CreateStatic(&FMainFrameActionCallbacks::PackageProjectCanExecute, Platform->PlatformName(), Platform->SupportsFeature(ETargetPlatformFeatures::Packaging))
+			FExecuteAction::CreateStatic(&FMainFrameActionCallbacks::PackageProject, PlatformInfo.PlatformInfoName),
+			FCanExecuteAction::CreateStatic(&FMainFrameActionCallbacks::PackageProjectCanExecute, PlatformInfo.PlatformInfoName, Platform && Platform->SupportsFeature(ETargetPlatformFeatures::Packaging))
 			);
 
 		// ... generate tooltip text
 		FFormatNamedArguments TooltipArguments;
-		TooltipArguments.Add(TEXT("DisplayName"), Platform->DisplayName());
+		TooltipArguments.Add(TEXT("DisplayName"), PlatformInfo.DisplayName);
 		FText Tooltip = FText::Format(LOCTEXT("PackageGameForPlatformTooltip", "Build, cook and package your game for the {DisplayName} platform"), TooltipArguments);
 
 		FProjectStatus ProjectStatus;
-		if (IProjectManager::Get().QueryStatusForCurrentProject(ProjectStatus) && !ProjectStatus.IsTargetPlatformSupported(*Platform->PlatformName()))
+		if (IProjectManager::Get().QueryStatusForCurrentProject(ProjectStatus) && !ProjectStatus.IsTargetPlatformSupported(PlatformInfo.VanillaPlatformName))
 		{
-			FText TooltipLine2 = FText::Format(LOCTEXT("LaunchDevicePlatformWarning", "{DisplayName} is not listed as a target platform for this project, so may not run as expected."), TooltipArguments);
+			FText TooltipLine2 = FText::Format(LOCTEXT("PackageUnsupportedPlatformWarning", "{DisplayName} is not listed as a target platform for this project, so may not run as expected."), TooltipArguments);
 			Tooltip = FText::Format(FText::FromString(TEXT("{0}\n\n{1}")), Tooltip, TooltipLine2);
 		}
 
 		// ... and add a menu entry
 		MenuBuilder.AddMenuEntry(
 			Action, 
-			ProjectTargetPlatformEditorModule.MakePlatformMenuItemWidget(Platform), 
+			ProjectTargetPlatformEditorModule.MakePlatformMenuItemWidget(PlatformInfo), 
 			NAME_None, 
 			Tooltip
 			);
@@ -118,42 +128,14 @@ protected:
 	* Creates the platform menu entries for a given platforms sub-platforms.
 	* e.g. Windows has multiple sub-platforms - Win32 and Win64
 	*
-	* @param MenuBuilder	- The builder for the menu that owns this menu.
-	* @param Platform		- The target platform we allow packaging for
-	* @param DisplayNames	- The Sub-platform display names
+	* @param MenuBuilder		- The builder for the menu that owns this menu.
+	* @param SubPlatformInfos	- The Sub-platform information
 	*/
-	static void AddPlatformSubPlatformsToMenu(FMenuBuilder& MenuBuilder, const ITargetPlatform* Platform, TArray<FText> DisplayNames)
+	static void AddPlatformSubPlatformsToMenu(FMenuBuilder& MenuBuilder, TArray<const PlatformInfo::FPlatformInfo*> SubPlatformInfos)
 	{
-		IProjectTargetPlatformEditorModule& ProjectTargetPlatformEditorModule = FModuleManager::LoadModuleChecked<IProjectTargetPlatformEditorModule>("ProjectTargetPlatformEditor");
-
-		for (const FText& DisplayName : DisplayNames)
+		for (const PlatformInfo::FPlatformInfo* SubPlatformInfo : SubPlatformInfos)
 		{
-			FString OptionalParams = FString::Printf(TEXT("-targetplatform=%s"), *DisplayName.ToString());
-			FUIAction Action(
-				FExecuteAction::CreateStatic(&FMainFrameActionCallbacks::PackageProject, Platform->PlatformName(), DisplayName, OptionalParams),
-				FCanExecuteAction::CreateStatic(&FMainFrameActionCallbacks::PackageProjectCanExecute, Platform->PlatformName(), Platform->SupportsFeature(ETargetPlatformFeatures::Packaging))
-				);
-
-			// ... generate tooltip text
-			FFormatNamedArguments TooltipArguments;
-			TooltipArguments.Add(TEXT("DisplayName"), Platform->DisplayName());
-			TooltipArguments.Add(TEXT("DisplaySubName"), DisplayName);
-			FText Tooltip = FText::Format(LOCTEXT("PackageGameForSubPlatformTooltip", "Build, cook and package your game for the {DisplayName} | {DisplaySubName} platform"), TooltipArguments);
-
-			FProjectStatus ProjectStatus;
-			if (IProjectManager::Get().QueryStatusForCurrentProject(ProjectStatus) && !ProjectStatus.IsTargetPlatformSupported(*Platform->PlatformName()))
-			{
-				FText TooltipLine2 = FText::Format(LOCTEXT("LaunchDevicePlatformWarning", "{DisplayName} is not listed as a target platform for this project, so may not run as expected."), TooltipArguments);
-				Tooltip = FText::Format(FText::FromString(TEXT("{0}\n\n{1}")), Tooltip, TooltipLine2);
-			}
-
-			// ... and add a menu entry
-			MenuBuilder.AddMenuEntry(
-				Action, 
-				ProjectTargetPlatformEditorModule.MakePlatformMenuItemWidget(Platform, false, DisplayName), 
-				NAME_None, 
-				Tooltip
-				);
+			AddPlatformToMenu(MenuBuilder, *SubPlatformInfo);
 		}
 	}
 
diff --git a/Engine/Source/Editor/ProjectTargetPlatformEditor/Private/ProjectTargetPlatformEditorModule.cpp b/Engine/Source/Editor/ProjectTargetPlatformEditor/Private/ProjectTargetPlatformEditorModule.cpp
--- a/Engine/Source/Editor/ProjectTargetPlatformEditor/Private/ProjectTargetPlatformEditorModule.cpp
+++ b/Engine/Source/Editor/ProjectTargetPlatformEditor/Private/ProjectTargetPlatformEditorModule.cpp
@@ -45,7 +45,7 @@ public:
 			);
 	}
 
-	virtual TSharedRef<SWidget> MakePlatformMenuItemWidget(const ITargetPlatform* const Platform, const bool bForCheckBox = false, const FText& DisplayNameOverride = FText()) const OVERRIDE
+	virtual TSharedRef<SWidget> MakePlatformMenuItemWidget(const PlatformInfo::FPlatformInfo& PlatformInfo, const bool bForCheckBox = false, const FText& DisplayNameOverride = FText()) const OVERRIDE
 	{
 		struct Local
 		{
@@ -74,7 +74,7 @@ public:
 					.HeightOverride(MultiBoxConstants::MenuIconSize)
 					[
 						SNew(SImage)
-						.Image(FEditorStyle::GetBrush(*FString::Printf(TEXT("Launcher.Platform_%s"), *Platform->PlatformName())))
+						.Image(FEditorStyle::GetBrush(PlatformInfo.GetIconStyleName(PlatformInfo::EPlatformIconSize::Normal)))
 					]
 				]
 				+SOverlay::Slot()
@@ -87,7 +87,7 @@ public:
 					.HeightOverride(MultiBoxConstants::MenuIconSize * 0.9f)
 					[
 						SNew(SImage)
-						.Visibility(TAttribute<EVisibility>::Create(TAttribute<EVisibility>::FGetter::CreateStatic(&Local::IsUnsupportedPlatformWarningVisible, FName(*Platform->PlatformName()))))
+						.Visibility(TAttribute<EVisibility>::Create(TAttribute<EVisibility>::FGetter::CreateStatic(&Local::IsUnsupportedPlatformWarningVisible, PlatformInfo.VanillaPlatformName)))
 						.Image(FEditorStyle::GetBrush("Launcher.Platform.Warning"))
 					]
 				]
@@ -99,21 +99,21 @@ public:
 			[
 				SNew(STextBlock)
 				.TextStyle(FEditorStyle::Get(), "Menu.Label")
-				.Text((DisplayNameOverride.IsEmpty()) ? Platform->DisplayName() : DisplayNameOverride)
+				.Text((DisplayNameOverride.IsEmpty()) ? PlatformInfo.DisplayName : DisplayNameOverride)
 			];
 	}
 
 	virtual bool ShowUnsupportedTargetWarning(const FName PlatformName) const OVERRIDE
 	{
+		const PlatformInfo::FPlatformInfo* const PlatformInfo = PlatformInfo::FindPlatformInfo(PlatformName);
+		check(PlatformInfo);
+
 		// Don't show the warning during automation testing; the dlg is modal and blocks
 		FProjectStatus ProjectStatus;
-		if(!GIsAutomationTesting && IProjectManager::Get().QueryStatusForCurrentProject(ProjectStatus) && !ProjectStatus.IsTargetPlatformSupported(PlatformName))
+		if(!GIsAutomationTesting && IProjectManager::Get().QueryStatusForCurrentProject(ProjectStatus) && !ProjectStatus.IsTargetPlatformSupported(PlatformInfo->VanillaPlatformName))
 		{
-			ITargetPlatform *const Platform = GetTargetPlatformManager()->FindTargetPlatform(PlatformName.ToString());
-			check(Platform);
-
 			FFormatNamedArguments Args;
-			Args.Add(TEXT("DisplayName"), Platform->DisplayName());
+			Args.Add(TEXT("DisplayName"), PlatformInfo->DisplayName);
 			FText WarningText = FText::Format(LOCTEXT("ShowUnsupportedPlatformWarning_Message", "{DisplayName} is not listed as a supported platform for this project, so may not run as expected.\n\nDo you wish to continue?"), Args);
 
 			FSuppressableWarningDialog::FSetupInfo Info(
diff --git a/Engine/Source/Editor/ProjectTargetPlatformEditor/Private/ProjectTargetPlatformEditorPrivatePCH.h b/Engine/Source/Editor/ProjectTargetPlatformEditor/Private/ProjectTargetPlatformEditorPrivatePCH.h
--- a/Engine/Source/Editor/ProjectTargetPlatformEditor/Private/ProjectTargetPlatformEditorPrivatePCH.h
+++ b/Engine/Source/Editor/ProjectTargetPlatformEditor/Private/ProjectTargetPlatformEditorPrivatePCH.h
@@ -13,6 +13,7 @@
 #include "ModuleManager.h"
 #include "EditorStyle.h"
 #include "TargetPlatform.h"
+#include "PlatformInfo.h"
 #include "GameProjectGenerationModule.h"
 
 
diff --git a/Engine/Source/Editor/ProjectTargetPlatformEditor/Private/Widgets/SProjectTargetPlatformSettings.cpp b/Engine/Source/Editor/ProjectTargetPlatformEditor/Private/Widgets/SProjectTargetPlatformSettings.cpp
--- a/Engine/Source/Editor/ProjectTargetPlatformEditor/Private/Widgets/SProjectTargetPlatformSettings.cpp
+++ b/Engine/Source/Editor/ProjectTargetPlatformEditor/Private/Widgets/SProjectTargetPlatformSettings.cpp
@@ -4,124 +4,33 @@
 
 #define LOCTEXT_NAMESPACE "SProjectTargetPlatformSettings"
 
-/**
- * Implements a row widget for the target platform list view
- */
-class SProjectTargetPlatformsListRow : public SMultiColumnTableRow<TSharedPtr<FProjectTargetPlatform>>
+BEGIN_SLATE_FUNCTION_BUILD_OPTIMIZATION
+void SProjectTargetPlatformSettings::Construct(const FArguments& InArgs)
 {
-public:
-
-	SLATE_BEGIN_ARGS(SProjectTargetPlatformsListRow)
-		: _TargetPlatform()
-	{}
-
-	SLATE_ARGUMENT(TSharedPtr<FProjectTargetPlatform>, TargetPlatform)
-
-	SLATE_END_ARGS()
-
-public:
-
-	/**
-	 * Constructs the widget.
-	 *
-	 * @param InArgs - The arguments.
-	 */
-	void Construct(const FArguments& InArgs, const TSharedRef<STableViewBase>& InOwnerTableView)
+	// Create and sort a list of vanilla platforms that are game targets (sort by display name)
+	// We show all of the platforms regardless of whether we have an SDK installed for them or not
+	for(const PlatformInfo::FPlatformInfo& PlatformInfo : PlatformInfo::EnumeratePlatformInfoArray())
 	{
-		check(InArgs._TargetPlatform.IsValid());
-		TargetPlatform = InArgs._TargetPlatform;
-
-		SMultiColumnTableRow<TSharedPtr<FProjectTargetPlatform>>::Construct(FSuperRowType::FArguments(), InOwnerTableView);
-	}
-
-public:
-
-	/**
-	 * Generates the widget for the specified column.
-	 *
-	 * @param ColumnName - The name of the column to generate the widget for.
-	 *
-	 * @return The widget.
-	 */
-	BEGIN_SLATE_FUNCTION_BUILD_OPTIMIZATION
-	virtual TSharedRef<SWidget> GenerateWidgetForColumn(const FName& ColumnName) OVERRIDE
-	{
-		if(ColumnName == "Enabled")
+		if(PlatformInfo.IsVanilla() && PlatformInfo.PlatformType == PlatformInfo::EPlatformType::Game)
 		{
-			return SNew(SBox)
-			.Padding(2.0f)
-			.VAlign(VAlign_Center)
-			[
-				SNew(SCheckBox)
-				.IsChecked(this, &SProjectTargetPlatformsListRow::HandleEnabledCheckBoxIsChecked)
-				.OnCheckStateChanged(this, &SProjectTargetPlatformsListRow::HandleEnabledCheckBoxStateChanged)
-			];
+			AvailablePlatforms.Add(&PlatformInfo);
 		}
-		else if(ColumnName == "Icon")
-		{
-			return SNew(SBox)
-			.Padding(FMargin(4.0f, 0.0f))
-			[
-				SNew(SImage)
-				.Image(TargetPlatform->Icon)
-			];
-		}
-		else if(ColumnName == "Platform")
-		{
-			return SNew(SBox)
-			.Padding(FMargin(4.0f, 0.0f))
-			.VAlign(VAlign_Center)
-			[
-				SNew(STextBlock)
-				.Text(TargetPlatform->DisplayName)
-				.ToolTipText(TargetPlatform->PlatformName.ToString())
-			];
-		}
-
-		return SNullWidget::NullWidget;
 	}
-	END_SLATE_FUNCTION_BUILD_OPTIMIZATION
 
-private:
-	/** Check to see if the "enabled" checkbox should be checked for this platform */
-	ESlateCheckBoxState::Type HandleEnabledCheckBoxIsChecked() const
+	AvailablePlatforms.Sort([](const PlatformInfo::FPlatformInfo& One, const PlatformInfo::FPlatformInfo& Two) -> bool
 	{
-		// We need to check that the platform is explicitly supported as a target, rather than implicitly supported due to an empty list
-		FProjectStatus ProjectStatus;
-		return (IProjectManager::Get().QueryStatusForCurrentProject(ProjectStatus) && ProjectStatus.IsTargetPlatformSupported(TargetPlatform->PlatformName, false/*bAllowSupportedIfEmptyList*/)) 
-			? ESlateCheckBoxState::Checked 
-			: ESlateCheckBoxState::Unchecked;
-	}
-
-	/** Handle the "enabled" checkbox state being changed for this platform */
-	void HandleEnabledCheckBoxStateChanged(ESlateCheckBoxState::Type InState)
-	{
-		FGameProjectGenerationModule::Get().UpdateSupportedTargetPlatforms(TargetPlatform->PlatformName, InState == ESlateCheckBoxState::Checked);
-	}
-
-	TSharedPtr<FProjectTargetPlatform> TargetPlatform;
-};
+		return One.DisplayName.CompareTo(Two.DisplayName) < 0;
+	});
 
-BEGIN_SLATE_FUNCTION_BUILD_OPTIMIZATION
-void SProjectTargetPlatformSettings::Construct(const FArguments& InArgs)
-{
+	// Generate a widget for each platform
+	TSharedRef<SVerticalBox> PlatformsListBox = SNew(SVerticalBox);
+	for(const PlatformInfo::FPlatformInfo* AvailablePlatform : AvailablePlatforms)
 	{
-		TArray<ITargetPlatform*> AllPlatforms = GetTargetPlatformManager()->GetTargetPlatforms();
-		for(ITargetPlatform* Platform : AllPlatforms)
-		{
-			// We are only interested in standalone games
-			const bool bIsGamePlatform = !Platform->IsClientOnly() && !Platform->IsServerOnly() && !Platform->HasEditorOnlyData();
-			if(bIsGamePlatform)
-			{
-				// We use FProjectTargetPlatform rather than ITargetPlatform* to avoid trying to access ITargetPlatform
-				// instances after their module has been unloaded (which can happen on editor shutdown)
-				AvailablePlatforms.Add(MakeShareable(new FProjectTargetPlatform(
-					Platform->DisplayName(), 
-					*Platform->PlatformName(), 
-					FEditorStyle::GetBrush(*FString::Printf(TEXT("Launcher.Platform_%s"), *Platform->PlatformName()))
-					)));
-			}
-		}
+		PlatformsListBox->AddSlot()
+		.AutoHeight()
+		[
+			MakePlatformRow(AvailablePlatform->DisplayName, AvailablePlatform->PlatformInfoName, FEditorStyle::GetBrush(AvailablePlatform->GetIconStyleName(PlatformInfo::EPlatformIconSize::Normal)))
+		];
 	}
 
 	ChildSlot
@@ -130,39 +39,6 @@ void SProjectTargetPlatformSettings::Construct(const FArguments& InArgs)
 
 		+SVerticalBox::Slot()
 		.AutoHeight()
-		[
-			SNew(SBorder)
-			.BorderImage(FEditorStyle::GetBrush("ToolPanel.GroupBorder"))
-			.Padding(5.0f)
-			[
-				SNew(SListView<TSharedPtr<FProjectTargetPlatform>>)
-				.ItemHeight(20.0f)
-				.ListItemsSource(&AvailablePlatforms)
-				.OnGenerateRow(this, &SProjectTargetPlatformSettings::HandlePlatformsListViewGenerateRow)
-				.SelectionMode(ESelectionMode::Single)
-				.HeaderRow
-				(
-					SNew(SHeaderRow)
-
-					+SHeaderRow::Column("Enabled")
-					.DefaultLabel(FText::FromString(TEXT(" ")))
-					.FixedWidth(28.0f)
-					.HAlignCell(HAlign_Center)
-					.HAlignHeader(HAlign_Center)
-
-					+SHeaderRow::Column("Icon")
-					.DefaultLabel(FText::FromString(TEXT(" ")))
-					.FixedWidth(32.0f)
-
-					+SHeaderRow::Column("Platform")
-					.DefaultLabel(LOCTEXT("PlatformsListPlatformColumnHeader", "Platform"))
-				)
-			]
-		]
-
-		+SVerticalBox::Slot()
-		.AutoHeight()
-		.Padding(FMargin(0.0f, 5.0f))
 		[
 			SNew(SBorder)
 			.BorderImage(FEditorStyle::GetBrush("ToolPanel.GroupBorder"))
@@ -173,28 +49,141 @@ void SProjectTargetPlatformSettings::Construct(const FArguments& InArgs)
 				+SVerticalBox::Slot()
 				.AutoHeight()
 				[
-					SNew(STextBlock)
-					.AutoWrapText(true)
-					.Text(LOCTEXT("PlatformsListDescription", "Select the supported platforms for your project. Attempting to package, run, or cook your project on an unsupported platform will result in a warning."))
+					MakePlatformRow(LOCTEXT("AllPlatforms", "All Platforms"), NAME_None, FEditorStyle::GetBrush("Launcher.Platform.AllPlatforms"))
 				]
 
 				+SVerticalBox::Slot()
 				.AutoHeight()
+				.Padding(FMargin(0.0f, 1.0f))
 				[
-					SNew(STextBlock)
-					.AutoWrapText(true)
-					.Text(LOCTEXT("PlatformsListNote", "Note: Selecting no platforms is considered the same as selecting all platforms."))
+					SNew(SSeparator)
+				]
+
+				+SVerticalBox::Slot()
+				.AutoHeight()
+				[
+					PlatformsListBox
 				]
 			]
 		]
+
+		+SVerticalBox::Slot()
+		.AutoHeight()
+		.Padding(FMargin(0.0f, 5.0f))
+		[
+			SNew(SBorder)
+			.BorderImage(FEditorStyle::GetBrush("ToolPanel.GroupBorder"))
+			.Padding(5.0f)
+			[
+				SNew(STextBlock)
+				.AutoWrapText(true)
+				.Text(LOCTEXT("PlatformsListDescription", "Select the supported platforms for your project. Attempting to package, run, or cook your project on an unsupported platform will result in a warning."))
+			]
+		]
 	];
 }
 END_SLATE_FUNCTION_BUILD_OPTIMIZATION
 
-TSharedRef<ITableRow> SProjectTargetPlatformSettings::HandlePlatformsListViewGenerateRow(TSharedPtr<FProjectTargetPlatform> TargetPlatform, const TSharedRef<STableViewBase>& OwnerTable)
+TSharedRef<SWidget> SProjectTargetPlatformSettings::MakePlatformRow(const FText& DisplayName, const FName PlatformName, const FSlateBrush* Icon) const
 {
-	return SNew(SProjectTargetPlatformsListRow, OwnerTable)
-		.TargetPlatform(TargetPlatform);
+	return SNew(SHorizontalBox)
+
+		+SHorizontalBox::Slot()
+		.AutoWidth()
+		.Padding(FMargin(5.0f, 0.0f))
+		[
+			SNew(SCheckBox)
+			.IsChecked(this, &SProjectTargetPlatformSettings::HandlePlatformCheckBoxIsChecked, PlatformName)
+			.IsEnabled(this, &SProjectTargetPlatformSettings::HandlePlatformCheckBoxIsEnabled, PlatformName)
+			.OnCheckStateChanged(this, &SProjectTargetPlatformSettings::HandlePlatformCheckBoxStateChanged, PlatformName)
+		]
+
+		+SHorizontalBox::Slot()
+		.AutoWidth()
+		.Padding(FMargin(5.0f, 2.0f))
+		.HAlign(HAlign_Center)
+		.VAlign(VAlign_Center)
+		[
+			SNew(SBox)
+			.WidthOverride(20)
+			.HeightOverride(20)
+			[
+				SNew(SImage)
+				.Image(Icon)
+			]
+		]
+
+		+SHorizontalBox::Slot()
+		.AutoWidth()
+		.Padding(FMargin(5.0f, 0.0f))
+		.VAlign(VAlign_Center)
+		[
+			SNew(STextBlock)
+			.Text(DisplayName)
+		];
+}
+
+ESlateCheckBoxState::Type SProjectTargetPlatformSettings::HandlePlatformCheckBoxIsChecked(const FName PlatformName) const
+{
+	FProjectStatus ProjectStatus;
+	if(IProjectManager::Get().QueryStatusForCurrentProject(ProjectStatus))
+	{
+		if(PlatformName.IsNone())
+		{
+			// None is "All Platforms"
+			return (ProjectStatus.SupportsAllPlatforms()) ? ESlateCheckBoxState::Checked : ESlateCheckBoxState::Unchecked;
+		}
+		else
+		{
+			return (ProjectStatus.IsTargetPlatformSupported(PlatformName)) ? ESlateCheckBoxState::Checked : ESlateCheckBoxState::Unchecked;
+		}
+	}
+
+	return ESlateCheckBoxState::Unchecked;
+}
+
+bool SProjectTargetPlatformSettings::HandlePlatformCheckBoxIsEnabled(const FName PlatformName) const
+{
+	FProjectStatus ProjectStatus;
+	if(IProjectManager::Get().QueryStatusForCurrentProject(ProjectStatus))
+	{
+		if(PlatformName.IsNone())
+		{
+			// None is "All Platforms"
+			return true;
+		}
+		else
+		{
+			// Individual platforms are only enabled when we're not supporting "All Platforms"
+			return !ProjectStatus.SupportsAllPlatforms();
+		}
+	}
+
+	return false;
+}
+
+void SProjectTargetPlatformSettings::HandlePlatformCheckBoxStateChanged(ESlateCheckBoxState::Type InState, const FName PlatformName) const
+{
+	if(PlatformName.IsNone())
+	{
+		// None is "All Platforms"
+		if(InState == ESlateCheckBoxState::Checked)
+		{
+			FGameProjectGenerationModule::Get().ClearSupportedTargetPlatforms();
+		}
+		else
+		{
+			// We've deselected "All Platforms", so manually select every available platform
+			for(const PlatformInfo::FPlatformInfo* AvailablePlatform : AvailablePlatforms)
+			{
+				FGameProjectGenerationModule::Get().UpdateSupportedTargetPlatforms(AvailablePlatform->TargetPlatformName, true);
+			}
+		}
+	}
+	else
+	{
+		FGameProjectGenerationModule::Get().UpdateSupportedTargetPlatforms(PlatformName, InState == ESlateCheckBoxState::Checked);
+	}
 }
 
 #undef LOCTEXT_NAMESPACE
diff --git a/Engine/Source/Editor/ProjectTargetPlatformEditor/Private/Widgets/SProjectTargetPlatformSettings.h b/Engine/Source/Editor/ProjectTargetPlatformEditor/Private/Widgets/SProjectTargetPlatformSettings.h
--- a/Engine/Source/Editor/ProjectTargetPlatformEditor/Private/Widgets/SProjectTargetPlatformSettings.h
+++ b/Engine/Source/Editor/ProjectTargetPlatformEditor/Private/Widgets/SProjectTargetPlatformSettings.h
@@ -2,25 +2,6 @@
 
 #pragma once
 
-struct FProjectTargetPlatform
-{
-	FProjectTargetPlatform(FText InDisplayName, FName InPlatformName, const FSlateBrush* InIcon)
-		: DisplayName(InDisplayName)
-		, PlatformName(InPlatformName)
-		, Icon(InIcon)
-	{
-	}
-
-	/** Localized display name to show in the UI, eg, Windows */
-	FText DisplayName;
-
-	/** Internal platform name to use as a key, eg, WindowsNoEditor */
-	FName PlatformName;
-
-	/** Icon representing this platform */
-	const FSlateBrush* Icon;
-};
-
 class SProjectTargetPlatformSettings : public SCompoundWidget
 {
 public:
@@ -32,9 +13,17 @@ public:
 	void Construct(const FArguments& InArgs);
 
 private:
-	/** Generate a row for the available platforms list */
-	TSharedRef<ITableRow> HandlePlatformsListViewGenerateRow(TSharedPtr<FProjectTargetPlatform> TargetPlatform, const TSharedRef<STableViewBase>& OwnerTable);
+	/** Generate an entry for the given platform information */
+	TSharedRef<SWidget> MakePlatformRow(const FText& DisplayName, const FName PlatformName, const FSlateBrush* Icon) const;
+
+	/** Check to see if the "enabled" checkbox should be checked for this platform */
+	ESlateCheckBoxState::Type HandlePlatformCheckBoxIsChecked(const FName PlatformName) const;
+
+	/** Check to see if the "enabled" checkbox should be enabled for this platform */
+	bool HandlePlatformCheckBoxIsEnabled(const FName PlatformName) const;
+
+	/** Handle the "enabled" checkbox state being changed for this platform */
+	void HandlePlatformCheckBoxStateChanged(ESlateCheckBoxState::Type InState, const FName PlatformName) const;
 
-	/** List of available target platforms */
-	TArray<TSharedPtr<FProjectTargetPlatform>> AvailablePlatforms;
+	TArray<const PlatformInfo::FPlatformInfo*> AvailablePlatforms;
 };
diff --git a/Engine/Source/Editor/ProjectTargetPlatformEditor/ProjectTargetPlatformEditor.Build.cs b/Engine/Source/Editor/ProjectTargetPlatformEditor/ProjectTargetPlatformEditor.Build.cs
--- a/Engine/Source/Editor/ProjectTargetPlatformEditor/ProjectTargetPlatformEditor.Build.cs
+++ b/Engine/Source/Editor/ProjectTargetPlatformEditor/ProjectTargetPlatformEditor.Build.cs
@@ -23,6 +23,7 @@ public class ProjectTargetPlatformEditor : ModuleRules
 				"EditorStyle",
 				"TargetPlatform",
 				"GameProjectGeneration",
+				"DesktopPlatform",
 				"Settings",
 				"UnrealEd",
 				"Projects",
diff --git a/Engine/Source/Editor/ProjectTargetPlatformEditor/Public/Interfaces/IProjectTargetPlatformEditorModule.h b/Engine/Source/Editor/ProjectTargetPlatformEditor/Public/Interfaces/IProjectTargetPlatformEditorModule.h
--- a/Engine/Source/Editor/ProjectTargetPlatformEditor/Public/Interfaces/IProjectTargetPlatformEditorModule.h
+++ b/Engine/Source/Editor/ProjectTargetPlatformEditor/Public/Interfaces/IProjectTargetPlatformEditorModule.h
@@ -2,6 +2,12 @@
 
 #pragma once
 
+namespace PlatformInfo
+{
+	// Forward declare type from DesktopPlatform rather than add an include dependency to everything using IProjectTargetPlatformEditorModule
+	struct FPlatformInfo;
+}
+
 /**
  * Interface for the platform target platform editor module
  */
@@ -33,11 +39,11 @@ public:
 	/**
 	 * Creates the widget to use for a platform entry within a FMenuBuilder.
 	 *
-	 * @param Platform - The target platform to build the widget for
+	 * @param PlatformInfo - The target platform info to build the widget for
 	 * @param bForCheckBox - true if the widget is for a checkbox menu item, false if it's for any other menu item (this affects the padding)
 	 * @param DisplayNameOverride - An alternate name to use for the platform, otherwise Platform->DisplayName() will be used
 	 */
-	virtual TSharedRef<SWidget> MakePlatformMenuItemWidget(const ITargetPlatform* const Platform, const bool bForCheckBox = false, const FText& DisplayNameOverride = FText()) const = 0;
+	virtual TSharedRef<SWidget> MakePlatformMenuItemWidget(const PlatformInfo::FPlatformInfo& PlatformInfo, const bool bForCheckBox = false, const FText& DisplayNameOverride = FText()) const = 0;
 
 	/**
 	 * Check to see if the given platform is on the list of supported targets, and show a warning message if it's not, allowing the user to continue or cancel
diff --git a/Engine/Source/Editor/UnrealEd/Private/Kismet2/DebuggerCommands.cpp b/Engine/Source/Editor/UnrealEd/Private/Kismet2/DebuggerCommands.cpp
--- a/Engine/Source/Editor/UnrealEd/Private/Kismet2/DebuggerCommands.cpp
+++ b/Engine/Source/Editor/UnrealEd/Private/Kismet2/DebuggerCommands.cpp
@@ -18,6 +18,7 @@
 
 #include "GameProjectGenerationModule.h"
 #include "ProjectTargetPlatformEditor.h"
+#include "PlatformInfo.h"
 
 //@TODO: Remove this dependency
 #include "Editor/LevelEditor/Public/LevelEditor.h"
@@ -554,9 +555,16 @@ TSharedRef< SWidget > FPlayWorldCommands::GenerateLaunchMenuContent( TSharedRef<
 	const bool bShouldCloseWindowAfterMenuSelection = true;
 	FMenuBuilder MenuBuilder(bShouldCloseWindowAfterMenuSelection, InCommandList);
 
+	TArray<ITargetPlatform*> Platforms = GetTargetPlatformManager()->GetTargetPlatforms();
+
+	// Sort the platforms by display name so that they appear in the same order as the other menus (like "Package Project" and "Cook Project")
+	Platforms.Sort([](ITargetPlatform& One, ITargetPlatform& Two) -> bool
+	{
+		return One.DisplayName().CompareTo(Two.DisplayName()) < 0;
+	});
+
 	// shared devices section
 	TSharedPtr<ITargetDeviceServicesModule> TargetDeviceServicesModule = StaticCastSharedPtr<ITargetDeviceServicesModule>(FModuleManager::Get().LoadModule(TEXT("TargetDeviceServices")));
-	TArray<ITargetPlatform*> Platforms = GetTargetPlatformManager()->GetTargetPlatforms();
 	IProjectTargetPlatformEditorModule& ProjectTargetPlatformEditorModule = FModuleManager::LoadModuleChecked<IProjectTargetPlatformEditorModule>("ProjectTargetPlatformEditor");
 
 	TArray<FString> PlatformsToMaybeInstallLinksFor;
@@ -565,12 +573,12 @@ TSharedRef< SWidget > FPlayWorldCommands::GenerateLaunchMenuContent( TSharedRef<
 
 	MenuBuilder.BeginSection("LevelEditorLaunchDevices", LOCTEXT("LaunchButtonDevicesSection", "Devices"));
 	{
-		for (int32 PlatformIndex = 0; PlatformIndex < Platforms.Num(); ++PlatformIndex)
+		for(ITargetPlatform* Platform : Platforms)
 		{
-			ITargetPlatform* Platform = Platforms[PlatformIndex];
+			const PlatformInfo::FPlatformInfo& PlatformInfo = Platform->GetPlatformInfo();
 
 			// for the Editor we are only interested in launching standalone games
-			if (Platform->IsClientOnly() || Platform->IsServerOnly() || Platform->HasEditorOnlyData())
+			if (PlatformInfo.PlatformType != PlatformInfo::EPlatformType::Game)
 			{
 				continue;
 			}
@@ -584,15 +592,13 @@ TSharedRef< SWidget > FPlayWorldCommands::GenerateLaunchMenuContent( TSharedRef<
 			{
 				FString NotInstalledDocLink;
 				// if the platform wasn't installed, we'll add a menu item later (we never care about code in this case, since we don't compile)
-				if (!Platform->IsSdkInstalled(false, NotInstalledDocLink) && PlatformsToMaybeInstallLinksFor.Find(Platform->DisplayName().ToString()) != INDEX_NONE)
+				if (!Platform->IsSdkInstalled(false, NotInstalledDocLink) && PlatformsToMaybeInstallLinksFor.Find(PlatformInfo.DisplayName.ToString()) != INDEX_NONE)
 				{
 					PlatformsToAddInstallLinksFor.Add(Platform);
 				}
 			}
 			else
 			{
-				FString IconName = FString::Printf(TEXT("Launcher.Platform_%s"), *Platform->PlatformName());
-
 				// for each proxy...
 				for (auto DeviceProxyIt = DeviceProxies.CreateIterator(); DeviceProxyIt; ++DeviceProxyIt)
 				{
@@ -627,11 +633,11 @@ TSharedRef< SWidget > FPlayWorldCommands::GenerateLaunchMenuContent( TSharedRef<
 					// ... generate tooltip text
 					FFormatNamedArguments TooltipArguments;
 					TooltipArguments.Add(TEXT("DeviceID"), FText::FromString(DeviceProxy->GetName()));
-					TooltipArguments.Add(TEXT("DisplayName"), Platform->DisplayName());
+					TooltipArguments.Add(TEXT("DisplayName"), PlatformInfo.DisplayName);
 					FText Tooltip = FText::Format(LOCTEXT("LaunchDeviceToolTipText", "Launch the game on this {DisplayName} device ({DeviceID})"), TooltipArguments);
 
 					FProjectStatus ProjectStatus;
-					if (IProjectManager::Get().QueryStatusForCurrentProject(ProjectStatus) && !ProjectStatus.IsTargetPlatformSupported(*Platform->PlatformName())) 
+					if (IProjectManager::Get().QueryStatusForCurrentProject(ProjectStatus) && !ProjectStatus.IsTargetPlatformSupported(PlatformInfo.VanillaPlatformName)) 
 					{
 						FText TooltipLine2 = FText::Format(LOCTEXT("LaunchDevicePlatformWarning", "{DisplayName} is not listed as a target platform for this project, so may not run as expected."), TooltipArguments);
 						Tooltip = FText::Format(FText::FromString(TEXT("{0}\n\n{1}")), Tooltip, TooltipLine2);
@@ -640,7 +646,7 @@ TSharedRef< SWidget > FPlayWorldCommands::GenerateLaunchMenuContent( TSharedRef<
 					// ... and add a menu entry
 					MenuBuilder.AddMenuEntry(
 						LaunchDeviceAction, 
-						ProjectTargetPlatformEditorModule.MakePlatformMenuItemWidget(Platform, true, Label), 
+						ProjectTargetPlatformEditorModule.MakePlatformMenuItemWidget(PlatformInfo, true, Label), 
 						NAME_None, 
 						Tooltip, 
 						EUserInterfaceActionType::Check
diff --git a/Engine/Source/Runtime/Core/Public/Android/AndroidProperties.h b/Engine/Source/Runtime/Core/Public/Android/AndroidProperties.h
--- a/Engine/Source/Runtime/Core/Public/Android/AndroidProperties.h
+++ b/Engine/Source/Runtime/Core/Public/Android/AndroidProperties.h
@@ -18,11 +18,6 @@
 struct FAndroidPlatformProperties
 	: public FGenericPlatformProperties
 {
-	static FORCEINLINE const char* DisplayName()
-	{
-		return "Android";
-	}
-
 	static FORCEINLINE const char* GetPhysicsFormat( )
 	{
 		return "PhysXPC";		//@todo android: physx format
@@ -101,11 +96,6 @@ struct FAndroidPlatformProperties
 
 struct FAndroid_PVRTCPlatformProperties : public FAndroidPlatformProperties
 {
-	static FORCEINLINE const char* DisplayName()
-	{
-		return "Android (PVRTC)";
-	}
-
 	static FORCEINLINE const char* PlatformName()
 	{
 		return "Android_PVRTC";
@@ -114,11 +104,6 @@ struct FAndroid_PVRTCPlatformProperties : public FAndroidPlatformProperties
 
 struct FAndroid_ATCCPlatformProperties : public FAndroidPlatformProperties
 {
-	static FORCEINLINE const char* DisplayName()
-	{
-		return "Android (ATC)";
-	}
-
 	static FORCEINLINE const char* PlatformName()
 	{
 		return "Android_ATC";
@@ -127,11 +112,6 @@ struct FAndroid_ATCCPlatformProperties : public FAndroidPlatformProperties
 
 struct FAndroid_DXTPlatformProperties : public FAndroidPlatformProperties
 {
-	static FORCEINLINE const char* DisplayName()
-	{
-		return "Android (DXT)";
-	}
-
 	static FORCEINLINE const char* PlatformName()
 	{
 		return "Android_DXT";
@@ -140,11 +120,6 @@ struct FAndroid_DXTPlatformProperties : public FAndroidPlatformProperties
 
 struct FAndroid_ETC1PlatformProperties : public FAndroidPlatformProperties
 {
-	static FORCEINLINE const char* DisplayName()
-	{
-		return "Android (ETC1)";
-	}
-
 	static FORCEINLINE const char* PlatformName()
 	{
 		return "Android_ETC1";
@@ -153,11 +128,6 @@ struct FAndroid_ETC1PlatformProperties : public FAndroidPlatformProperties
 
 struct FAndroid_ETC2PlatformProperties : public FAndroidPlatformProperties
 {
-	static FORCEINLINE const char* DisplayName()
-	{
-		return "Android (ETC2)";
-	}
-
 	static FORCEINLINE const char* PlatformName()
 	{
 		return "Android_ETC2";
diff --git a/Engine/Source/Runtime/Core/Public/GenericPlatform/GenericPlatformProperties.h b/Engine/Source/Runtime/Core/Public/GenericPlatform/GenericPlatformProperties.h
--- a/Engine/Source/Runtime/Core/Public/GenericPlatform/GenericPlatformProperties.h
+++ b/Engine/Source/Runtime/Core/Public/GenericPlatform/GenericPlatformProperties.h
@@ -16,16 +16,6 @@
  */
 struct FGenericPlatformProperties
 {
-	/**
-	 * Gets the platform's user friendly display name.
-	 *
-	 * @return The display name string.
-	 */
-	static FORCEINLINE const char* DisplayName()
-	{
-		return "";
-	}
-
 	/**
 	 * Gets the platform's physics format.
 	 *
diff --git a/Engine/Source/Runtime/Core/Public/HTML5/HTML5PlatformProperties.h b/Engine/Source/Runtime/Core/Public/HTML5/HTML5PlatformProperties.h
--- a/Engine/Source/Runtime/Core/Public/HTML5/HTML5PlatformProperties.h
+++ b/Engine/Source/Runtime/Core/Public/HTML5/HTML5PlatformProperties.h
@@ -18,11 +18,6 @@
 struct FHTML5PlatformProperties
 	: public FGenericPlatformProperties
 {
-	static FORCEINLINE const char* DisplayName()
-	{
-		return "HTML5";
-	}
-
 	static FORCEINLINE const char* GetPhysicsFormat( )
 	{
 		return "PhysXPC";
diff --git a/Engine/Source/Runtime/Core/Public/IOS/IOSPlatformProperties.h b/Engine/Source/Runtime/Core/Public/IOS/IOSPlatformProperties.h
--- a/Engine/Source/Runtime/Core/Public/IOS/IOSPlatformProperties.h
+++ b/Engine/Source/Runtime/Core/Public/IOS/IOSPlatformProperties.h
@@ -18,11 +18,6 @@
 struct FIOSPlatformProperties
 	: public FGenericPlatformProperties
 {
-	static FORCEINLINE const char* DisplayName()
-	{
-		return "iOS";
-	}
-
 	static FORCEINLINE bool HasEditorOnlyData( )
 	{
 		return false;
diff --git a/Engine/Source/Runtime/Core/Public/Linux/LinuxPlatformProperties.h b/Engine/Source/Runtime/Core/Public/Linux/LinuxPlatformProperties.h
--- a/Engine/Source/Runtime/Core/Public/Linux/LinuxPlatformProperties.h
+++ b/Engine/Source/Runtime/Core/Public/Linux/LinuxPlatformProperties.h
@@ -19,26 +19,6 @@ template<bool HAS_EDITOR_DATA, bool IS_DEDICATED_SERVER, bool IS_CLIENT_ONLY>
 struct FLinuxPlatformProperties
 	: public FGenericPlatformProperties
 {
-	static FORCEINLINE const char* DisplayName()
-	{
-		if (IS_DEDICATED_SERVER)
-		{
-			return "Linux (Dedicated Server)";
-		}
-
-		if (HAS_EDITOR_DATA)
-		{
-			return "Linux (Editor)";
-		}
-
-		if (IS_CLIENT_ONLY)
-		{
-			return "Linux (Client-only)";
-		}
-
-		return "Linux";
-	}
-
 	static FORCEINLINE bool HasEditorOnlyData( )
 	{
 		return HAS_EDITOR_DATA;
diff --git a/Engine/Source/Runtime/Core/Public/Mac/MacPlatformProperties.h b/Engine/Source/Runtime/Core/Public/Mac/MacPlatformProperties.h
--- a/Engine/Source/Runtime/Core/Public/Mac/MacPlatformProperties.h
+++ b/Engine/Source/Runtime/Core/Public/Mac/MacPlatformProperties.h
@@ -19,26 +19,6 @@ template<bool HAS_EDITOR_DATA, bool IS_DEDICATED_SERVER, bool IS_CLIENT_ONLY>
 struct FMacPlatformProperties
 	: public FGenericPlatformProperties
 {
-	static FORCEINLINE const char* DisplayName()
-	{
-		if (IS_DEDICATED_SERVER)
-		{
-			return "Mac (Dedicated Server)";
-		}
-		
-		if (HAS_EDITOR_DATA)
-		{
-			return "Mac (Editor)";
-		}
-
-		if (IS_CLIENT_ONLY)
-		{
-			return "Mac (Client-only)";
-		}
-
-		return "Mac";
-	}
-
 	static FORCEINLINE bool HasEditorOnlyData( )
 	{
 		return HAS_EDITOR_DATA;
diff --git a/Engine/Source/Runtime/Core/Public/WinRT/WinRTProperties.h b/Engine/Source/Runtime/Core/Public/WinRT/WinRTProperties.h
--- a/Engine/Source/Runtime/Core/Public/WinRT/WinRTProperties.h
+++ b/Engine/Source/Runtime/Core/Public/WinRT/WinRTProperties.h
@@ -13,11 +13,6 @@
 
 struct FWinRTPlatformProperties : public FGenericPlatformProperties
 {
-	static FORCEINLINE const char* DisplayName()
-	{
-		return "Windows RT";
-	}
-
 	static FORCEINLINE const char* PlatformName()
 	{
 		return "WinRT";
diff --git a/Engine/Source/Runtime/Core/Public/Windows/WindowsPlatformProperties.h b/Engine/Source/Runtime/Core/Public/Windows/WindowsPlatformProperties.h
--- a/Engine/Source/Runtime/Core/Public/Windows/WindowsPlatformProperties.h
+++ b/Engine/Source/Runtime/Core/Public/Windows/WindowsPlatformProperties.h
@@ -19,26 +19,6 @@ template<bool HAS_EDITOR_DATA, bool IS_DEDICATED_SERVER, bool IS_CLIENT_ONLY>
 struct FWindowsPlatformProperties
 	: public FGenericPlatformProperties
 {
-	static FORCEINLINE const char* DisplayName()
-	{
-		if (IS_DEDICATED_SERVER)
-		{
-			return "Windows (Dedicated Server)";
-		}
-		
-		if (HAS_EDITOR_DATA)
-		{
-			return "Windows (Editor)";
-		}
-		
-		if (IS_CLIENT_ONLY)
-		{
-			return "Windows (Client-only)";
-		}
-
-		return "Windows";
-	}
-
 	static FORCEINLINE bool HasEditorOnlyData()
 	{
 		return HAS_EDITOR_DATA;
diff --git a/Engine/Source/Runtime/Projects/Private/ProjectManager.cpp b/Engine/Source/Runtime/Projects/Private/ProjectManager.cpp
--- a/Engine/Source/Runtime/Projects/Private/ProjectManager.cpp
+++ b/Engine/Source/Runtime/Projects/Private/ProjectManager.cpp
@@ -40,6 +40,11 @@ void FProject::UpdateSupportedTargetPlatforms(const FName& InPlatformName, const
 	}
 }
 
+void FProject::ClearSupportedTargetPlatforms()
+{
+	ProjectInfo.TargetPlatforms.Empty();
+}
+
 bool FProject::PerformAdditionalDeserialization(const TSharedRef< FJsonObject >& FileObject)
 {
 	ReadNumberFromJSON(FileObject, TEXT("EpicSampleNameHash"), ProjectInfo.EpicSampleNameHash);
@@ -362,6 +367,44 @@ void FProjectManager::UpdateSupportedTargetPlatformsForCurrentProject(const FNam
 	OnTargetPlatformsForCurrentProjectChangedEvent.Broadcast();
 }
 
+void FProjectManager::ClearSupportedTargetPlatformsForProject(const FString& FilePath)
+{
+	TSharedRef<FProject> NewProject = MakeShareable( new FProject() );
+	FText FailReason;
+	if ( !NewProject->LoadFromFile(FilePath, FailReason) )
+	{
+		return;
+	}
+
+	NewProject->ClearSupportedTargetPlatforms();
+
+	const FString& FileContents = NewProject->SerializeToJSON();
+	FFileHelper::SaveStringToFile(FileContents, *FilePath);
+
+	// Call OnTargetPlatformsForCurrentProjectChangedEvent if this project is the same as the one we currently have loaded
+	const FString CurrentProjectPath = FPaths::ConvertRelativePathToFull(FPaths::GetProjectFilePath());
+	const FString InProjectPath = FPaths::ConvertRelativePathToFull(FilePath);
+	if ( CurrentProjectPath == InProjectPath )
+	{
+		OnTargetPlatformsForCurrentProjectChangedEvent.Broadcast();
+	}
+}
+
+void FProjectManager::ClearSupportedTargetPlatformsForCurrentProject()
+{
+	if ( !CurrentlyLoadedProject.IsValid() )
+	{
+		return;
+	}
+
+	CurrentlyLoadedProject->ClearSupportedTargetPlatforms();
+
+	const FString& FileContents = CurrentlyLoadedProject->SerializeToJSON();
+	FFileHelper::SaveStringToFile(FileContents, *FPaths::GetProjectFilePath());
+
+	OnTargetPlatformsForCurrentProjectChangedEvent.Broadcast();
+}
+
 IProjectManager& IProjectManager::Get()
 {
 	// Single instance of manager, allocated on demand and destroyed on program exit.
diff --git a/Engine/Source/Runtime/Projects/Private/ProjectManager.h b/Engine/Source/Runtime/Projects/Private/ProjectManager.h
--- a/Engine/Source/Runtime/Projects/Private/ProjectManager.h
+++ b/Engine/Source/Runtime/Projects/Private/ProjectManager.h
@@ -65,6 +65,9 @@ public:
 	 */
 	void UpdateSupportedTargetPlatforms(const FName& InPlatformName, const bool bIsSupported);
 
+	/** Clear the list of supported target platforms */
+	void ClearSupportedTargetPlatforms();
+
 	const FProjectInfo& GetProjectInfo() const
 	{
 		return ProjectInfo;
@@ -105,6 +108,8 @@ public:
 	virtual bool QueryStatusForCurrentProject(FProjectStatus& OutProjectStatus) const OVERRIDE;
 	virtual void UpdateSupportedTargetPlatformsForProject(const FString& FilePath, const FName& InPlatformName, const bool bIsSupported) OVERRIDE;
 	virtual void UpdateSupportedTargetPlatformsForCurrentProject(const FName& InPlatformName, const bool bIsSupported) OVERRIDE;
+	virtual void ClearSupportedTargetPlatformsForProject(const FString& FilePath) OVERRIDE;
+	virtual void ClearSupportedTargetPlatformsForCurrentProject() OVERRIDE;
 	virtual FOnTargetPlatformsForCurrentProjectChangedEvent& OnTargetPlatformsForCurrentProjectChanged() OVERRIDE { return OnTargetPlatformsForCurrentProjectChangedEvent; }
 
 private:
diff --git a/Engine/Source/Runtime/Projects/Public/Interfaces/IProjectManager.h b/Engine/Source/Runtime/Projects/Public/Interfaces/IProjectManager.h
--- a/Engine/Source/Runtime/Projects/Public/Interfaces/IProjectManager.h
+++ b/Engine/Source/Runtime/Projects/Public/Interfaces/IProjectManager.h
@@ -51,6 +51,14 @@ public:
 		// An empty list is considered the same as supporting all platforms
 		return (bAllowSupportedIfEmptyList && TargetPlatforms.Num() == 0) || TargetPlatforms.Contains(InPlatformName);
 	}
+
+	/**
+	 * Check to see if the given the current project supports all platforms
+	 */
+	bool SupportsAllPlatforms() const
+	{
+		return TargetPlatforms.Num() == 0;
+	}
 };
 
 /**
@@ -194,6 +202,12 @@ public:
 	 */
 	virtual void UpdateSupportedTargetPlatformsForCurrentProject(const FName& InPlatformName, const bool bIsSupported) = 0;
 
+	/** Clear the list of supported target platforms for the target project */
+	virtual void ClearSupportedTargetPlatformsForProject(const FString& FilePath) = 0;
+
+	/** Clear the list of supported target platforms for the current project */
+	virtual void ClearSupportedTargetPlatformsForCurrentProject() = 0;
+
 	/** Called when the target platforms for the current project are changed */
 	DECLARE_MULTICAST_DELEGATE(FOnTargetPlatformsForCurrentProjectChangedEvent);
 	virtual FOnTargetPlatformsForCurrentProjectChangedEvent& OnTargetPlatformsForCurrentProjectChanged() = 0;
