Fix paths search for projects. FindFiles does not return names including the search path stem.

diff --git a/Engine/Source/Developer/DesktopPlatform/Private/DesktopPlatformBase.cpp b/Engine/Source/Developer/DesktopPlatform/Private/DesktopPlatformBase.cpp
--- a/Engine/Source/Developer/DesktopPlatform/Private/DesktopPlatformBase.cpp
+++ b/Engine/Source/Developer/DesktopPlatform/Private/DesktopPlatformBase.cpp
@@ -639,11 +639,11 @@ bool FDesktopPlatformBase::EnumerateProjectsKnownByEngine(const FString &Identif
 		for(int32 FolderIdx = 0; FolderIdx < ProjectFolders.Num(); FolderIdx++)
 		{
 			TArray<FString> ProjectFiles;
-			IFileManager::Get().FindFiles(ProjectFiles, *(ProjectFolders[FolderIdx] / TEXT("*.uproject")), true, false);
+			IFileManager::Get().FindFiles(ProjectFiles, *(SearchDirectories[Idx] / ProjectFolders[FolderIdx] / TEXT("*.uproject")), true, false);
 
 			for(int32 FileIdx = 0; FileIdx < ProjectFiles.Num(); FileIdx++)
 			{
-				OutProjectFileNames.AddUnique(ProjectFiles[FileIdx]);
+				OutProjectFileNames.AddUnique(SearchDirectories[Idx] / ProjectFolders[FolderIdx] / ProjectFiles[FileIdx]);
 			}
 		}
 	}
