Linux: add (stubbed) DesktopPlatform class.

Part of PR 140 (with mods), contributed by [REDACTED_PERSON] and #ue4linux community.

diff --git a/Engine/Source/Developer/DesktopPlatform/Private/DesktopPlatformPrivatePCH.h b/Engine/Source/Developer/DesktopPlatform/Private/DesktopPlatformPrivatePCH.h
--- a/Engine/Source/Developer/DesktopPlatform/Private/DesktopPlatformPrivatePCH.h
+++ b/Engine/Source/Developer/DesktopPlatform/Private/DesktopPlatformPrivatePCH.h
@@ -12,4 +12,6 @@ DECLARE_LOG_CATEGORY_EXTERN(LogDesktopPlatform, Log, All);
 #include "Windows/DesktopPlatformWindows.h"
 #elif PLATFORM_MAC
 #include "Mac/DesktopPlatformMac.h"
-#endif
\ No newline at end of file
+#elif PLATFORM_LINUX
+#include "Linux/DesktopPlatformLinux.h"
+#endif
diff --git a/Engine/Source/Developer/DesktopPlatform/Private/Linux/DesktopPlatformLinux.cpp b/Engine/Source/Developer/DesktopPlatform/Private/Linux/DesktopPlatformLinux.cpp
new file mode 100644
--- /dev/null
+++ b/Engine/Source/Developer/DesktopPlatform/Private/Linux/DesktopPlatformLinux.cpp
@@ -0,0 +1,77 @@
+// Copyright 1998-2014 Epic Games, Inc. All Rights Reserved.
+
+#include "DesktopPlatformPrivatePCH.h"
+
+#define LOCTEXT_NAMESPACE "DesktopPlatform"
+
+bool FDesktopPlatformLinux::OpenFileDialog(const void* ParentWindowHandle, const FString& DialogTitle, const FString& DefaultPath, const FString& DefaultFile, const FString& FileTypes, uint32 Flags, TArray<FString>& OutFilenames)
+{
+	unimplemented();
+	return false;
+}
+
+bool FDesktopPlatformLinux::SaveFileDialog(const void* ParentWindowHandle, const FString& DialogTitle, const FString& DefaultPath, const FString& DefaultFile, const FString& FileTypes, uint32 Flags, TArray<FString>& OutFilenames)
+{
+	unimplemented();
+	return false;
+}
+
+bool FDesktopPlatformLinux::OpenDirectoryDialog(const void* ParentWindowHandle, const FString& DialogTitle, const FString& DefaultPath, FString& OutFolderName)
+{
+	unimplemented();
+	return false;
+}
+
+bool FDesktopPlatformLinux::OpenFontDialog(const void* ParentWindowHandle, FString& OutFontName, float& OutHeight, EFontImportFlags::Type& OutFlags)
+{
+	unimplemented();
+	return false;
+}
+
+bool FDesktopPlatformLinux::OpenLauncher(bool Install, FString CommandLineParams )
+{
+	unimplemented();
+	return false;
+}
+
+bool FDesktopPlatformLinux::FileDialogShared(bool bSave, const void* ParentWindowHandle, const FString& DialogTitle, const FString& DefaultPath, const FString& DefaultFile, const FString& FileTypes, uint32 Flags, TArray<FString>& OutFilenames)
+{
+	unimplemented();
+	return false;
+}
+
+bool FDesktopPlatformLinux::RegisterEngineInstallation(const FString &RootDir, FString &OutIdentifier)
+{
+	return false;
+}
+
+void FDesktopPlatformLinux::EnumerateEngineInstallations(TMap<FString, FString> &OutInstallations)
+{
+	unimplemented();
+}
+
+bool FDesktopPlatformLinux::VerifyFileAssociations()
+{
+	unimplemented();
+	return false;
+}
+
+bool FDesktopPlatformLinux::UpdateFileAssociations()
+{
+	unimplemented();
+	return false;
+}
+
+bool FDesktopPlatformLinux::RunUnrealBuildTool(const FText& Description, const FString& RootDir, const FString& Arguments, FFeedbackContext* Warn)
+{
+	unimplemented();
+	return false;
+}
+
+FFeedbackContext* FDesktopPlatformLinux::GetNativeFeedbackContext()
+{
+	unimplemented();
+	return NULL;
+}
+
+#undef LOCTEXT_NAMESPACE
diff --git a/Engine/Source/Developer/DesktopPlatform/Private/Linux/DesktopPlatformLinux.h b/Engine/Source/Developer/DesktopPlatform/Private/Linux/DesktopPlatformLinux.h
new file mode 100644
--- /dev/null
+++ b/Engine/Source/Developer/DesktopPlatform/Private/Linux/DesktopPlatformLinux.h
@@ -0,0 +1,31 @@
+// Copyright 1998-2014 Epic Games, Inc. All Rights Reserved.
+
+#pragma once
+
+#include "../DesktopPlatformBase.h"
+
+class FDesktopPlatformLinux : public FDesktopPlatformBase
+{
+public:
+	// IDesktopPlatform Implementation
+	virtual bool OpenFileDialog(const void* ParentWindowHandle, const FString& DialogTitle, const FString& DefaultPath, const FString& DefaultFile, const FString& FileTypes, uint32 Flags, TArray<FString>& OutFilenames) OVERRIDE;
+	virtual bool SaveFileDialog(const void* ParentWindowHandle, const FString& DialogTitle, const FString& DefaultPath, const FString& DefaultFile, const FString& FileTypes, uint32 Flags, TArray<FString>& OutFilenames) OVERRIDE;
+	virtual bool OpenDirectoryDialog(const void* ParentWindowHandle, const FString& DialogTitle, const FString& DefaultPath, FString& OutFolderName) OVERRIDE;
+	virtual bool OpenFontDialog(const void* ParentWindowHandle, FString& OutFontName, float& OutHeight, EFontImportFlags::Type& OutFlags) OVERRIDE;
+	virtual bool OpenLauncher(bool Install, FString CommandLineParams ) OVERRIDE;
+
+	virtual bool RegisterEngineInstallation(const FString &RootDir, FString &OutIdentifier) OVERRIDE;
+	virtual void EnumerateEngineInstallations(TMap<FString, FString> &OutInstallations) OVERRIDE;
+
+	virtual bool VerifyFileAssociations() OVERRIDE;
+	virtual bool UpdateFileAssociations() OVERRIDE;
+
+	virtual bool RunUnrealBuildTool(const FText& Description, const FString& RootDir, const FString& Arguments, FFeedbackContext* Warn) OVERRIDE;
+
+	virtual FFeedbackContext* GetNativeFeedbackContext() OVERRIDE;
+
+private:
+	bool FileDialogShared(bool bSave, const void* ParentWindowHandle, const FString& DialogTitle, const FString& DefaultPath, const FString& DefaultFile, const FString& FileTypes, uint32 Flags, TArray<FString>& OutFilenames);
+};
+
+typedef FDesktopPlatformLinux FDesktopPlatform;
