Prevent log spam when scanning for projects from the project browser, and .uproject files have missing or empty EngineAssociation entries.

diff --git a/Engine/Source/Developer/DesktopPlatform/Private/DesktopPlatformBase.cpp b/Engine/Source/Developer/DesktopPlatform/Private/DesktopPlatformBase.cpp
--- a/Engine/Source/Developer/DesktopPlatform/Private/DesktopPlatformBase.cpp
+++ b/Engine/Source/Developer/DesktopPlatform/Private/DesktopPlatformBase.cpp
@@ -247,6 +247,8 @@ bool FDesktopPlatformBase::SetEngineIdentifierForProject(const FString &ProjectF
 
 bool FDesktopPlatformBase::GetEngineIdentifierForProject(const FString &ProjectFileName, FString &OutIdentifier)
 {
+	OutIdentifier.Empty();
+
 	// Load the project file
 	TSharedPtr<FJsonObject> ProjectFile = LoadProjectFile(ProjectFileName);
 	if(!ProjectFile.IsValid())
@@ -254,11 +256,15 @@ bool FDesktopPlatformBase::GetEngineIdentifierForProject(const FString &ProjectF
 		return false;
 	}
 
-	// Read the identifier from it
-	OutIdentifier = ProjectFile->GetStringField(TEXT("EngineAssociation"));
-	if(OutIdentifier.Len() > 0)
+	// Try to read the identifier from it
+	TSharedPtr<FJsonValue> Value = ProjectFile->TryGetField(TEXT("EngineAssociation"));
+	if(Value.IsValid() && Value->Type == EJson::String)
 	{
-		return true;
+		OutIdentifier = Value->AsString();
+		if(OutIdentifier.Len() > 0)
+		{
+			return true;
+		}
 	}
 
 	// Otherwise scan up through the directory hierarchy to find an installation
