Fix Android flavor UAT commands

They need to specify -targetplatform=android as well as the -cookflavor

diff --git a/Engine/Source/Developer/DesktopPlatform/Private/PlatformInfo.cpp b/Engine/Source/Developer/DesktopPlatform/Private/PlatformInfo.cpp
--- a/Engine/Source/Developer/DesktopPlatform/Private/PlatformInfo.cpp
+++ b/Engine/Source/Developer/DesktopPlatform/Private/PlatformInfo.cpp
@@ -68,11 +68,11 @@ static const FPlatformInfo PlatformInfoArray[] = {
 	
 	BuildPlatformInfo(TEXT("Android"),					TEXT("Android"),			LOCTEXT("Android", "Android"),								EPlatformType::Game,		EPlatformFlags::None,			FPlatformIconPaths(TEXT("Launcher/Android/Platform_Android_24x"), TEXT("Launcher/Android/Platform_Android_128x")),						TEXT("")),
 	BuildPlatformInfo(TEXT("Android_All"),				TEXT("Android"),			LOCTEXT("Android_All", "Android (All)"),					EPlatformType::Game,		EPlatformFlags::CookFlavor,		FPlatformIconPaths(TEXT("Launcher/Android/Platform_Android_24x"), TEXT("Launcher/Android/Platform_Android_128x")),						TEXT("")),
-	BuildPlatformInfo(TEXT("Android_ATC"),				TEXT("Android_ATC"),		LOCTEXT("Android_ATC", "Android (ATC)"),					EPlatformType::Game,		EPlatformFlags::CookFlavor,		FPlatformIconPaths(TEXT("Launcher/Android/Platform_Android_ATC_24x"), TEXT("Launcher/Android/Platform_Android_128x")),					TEXT("-cookflavor=ATC")),
-	BuildPlatformInfo(TEXT("Android_DXT"),				TEXT("Android_DXT"),		LOCTEXT("Android_DXT", "Android (DXT)"),					EPlatformType::Game,		EPlatformFlags::CookFlavor,		FPlatformIconPaths(TEXT("Launcher/Android/Platform_Android_DXT_24x"), TEXT("Launcher/Android/Platform_Android_128x")),					TEXT("-cookflavor=DXT")),
-	BuildPlatformInfo(TEXT("Android_ETC1"),				TEXT("Android_ETC1"),		LOCTEXT("Android_ETC1", "Android (ETC1)"),					EPlatformType::Game,		EPlatformFlags::CookFlavor,		FPlatformIconPaths(TEXT("Launcher/Android/Platform_Android_ETC1_24x"), TEXT("Launcher/Android/Platform_Android_128x")),					TEXT("-cookflavor=ETC1")),
-	BuildPlatformInfo(TEXT("Android_ETC2"),				TEXT("Android_ETC2"),		LOCTEXT("Android_ETC2", "Android (ETC2)"),					EPlatformType::Game,		EPlatformFlags::CookFlavor,		FPlatformIconPaths(TEXT("Launcher/Android/Platform_Android_ETC2_24x"), TEXT("Launcher/Android/Platform_Android_128x")),					TEXT("-cookflavor=ETC2")),
-	BuildPlatformInfo(TEXT("Android_PVRTC"),			TEXT("Android_PVRTC"),		LOCTEXT("Android_PVRTC", "Android (PVRTC)"),				EPlatformType::Game,		EPlatformFlags::CookFlavor,		FPlatformIconPaths(TEXT("Launcher/Android/Platform_Android_PVRTC_24x"), TEXT("Launcher/Android/Platform_Android_128x")),				TEXT("-cookflavor=PVRTC")),
+	BuildPlatformInfo(TEXT("Android_ATC"),				TEXT("Android_ATC"),		LOCTEXT("Android_ATC", "Android (ATC)"),					EPlatformType::Game,		EPlatformFlags::CookFlavor,		FPlatformIconPaths(TEXT("Launcher/Android/Platform_Android_ATC_24x"), TEXT("Launcher/Android/Platform_Android_128x")),					TEXT("-targetplatform=Android -cookflavor=ATC")),
+	BuildPlatformInfo(TEXT("Android_DXT"),				TEXT("Android_DXT"),		LOCTEXT("Android_DXT", "Android (DXT)"),					EPlatformType::Game,		EPlatformFlags::CookFlavor,		FPlatformIconPaths(TEXT("Launcher/Android/Platform_Android_DXT_24x"), TEXT("Launcher/Android/Platform_Android_128x")),					TEXT("-targetplatform=Android -cookflavor=DXT")),
+	BuildPlatformInfo(TEXT("Android_ETC1"),				TEXT("Android_ETC1"),		LOCTEXT("Android_ETC1", "Android (ETC1)"),					EPlatformType::Game,		EPlatformFlags::CookFlavor,		FPlatformIconPaths(TEXT("Launcher/Android/Platform_Android_ETC1_24x"), TEXT("Launcher/Android/Platform_Android_128x")),					TEXT("-targetplatform=Android -cookflavor=ETC1")),
+	BuildPlatformInfo(TEXT("Android_ETC2"),				TEXT("Android_ETC2"),		LOCTEXT("Android_ETC2", "Android (ETC2)"),					EPlatformType::Game,		EPlatformFlags::CookFlavor,		FPlatformIconPaths(TEXT("Launcher/Android/Platform_Android_ETC2_24x"), TEXT("Launcher/Android/Platform_Android_128x")),					TEXT("-targetplatform=Android -cookflavor=ETC2")),
+	BuildPlatformInfo(TEXT("Android_PVRTC"),			TEXT("Android_PVRTC"),		LOCTEXT("Android_PVRTC", "Android (PVRTC)"),				EPlatformType::Game,		EPlatformFlags::CookFlavor,		FPlatformIconPaths(TEXT("Launcher/Android/Platform_Android_PVRTC_24x"), TEXT("Launcher/Android/Platform_Android_128x")),				TEXT("-targetplatform=Android -cookflavor=PVRTC")),
 
 	BuildPlatformInfo(TEXT("IOS"),						TEXT("IOS"),				LOCTEXT("IOS", "iOS"),										EPlatformType::Game,		EPlatformFlags::None,			FPlatformIconPaths(TEXT("Launcher/iOS/Platform_iOS_24x"), TEXT("Launcher/iOS/Platform_iOS_128x")),										TEXT("")),
 
