Track and reset the key window when using the OS X open/save dialog otherwise it will steal focus and not return it.

diff --git a/Engine/Source/Developer/DesktopPlatform/Private/Mac/DesktopPlatformMac.cpp b/Engine/Source/Developer/DesktopPlatform/Private/Mac/DesktopPlatformMac.cpp
--- a/Engine/Source/Developer/DesktopPlatform/Private/Mac/DesktopPlatformMac.cpp
+++ b/Engine/Source/Developer/DesktopPlatform/Private/Mac/DesktopPlatformMac.cpp
@@ -472,6 +472,7 @@ bool FDesktopPlatformMac::FileDialogShared(bool bSave, const void* ParentWindowH
 	[AccessoryView AddAllowedFileTypes:AllowedFileTypes];
 
 	bool bSuccess = false;
+	NSWindow* FocusWindow = [[NSApplication sharedApplication] keyWindow];
 
 	{
 		FScopedSystemModalMode SystemModalScope;
@@ -511,6 +512,11 @@ bool FDesktopPlatformMac::FileDialogShared(bool bSave, const void* ParentWindowH
 	[Panel close];
 
 	MacApplication->ResetModifierKeys();
+	
+	if(FocusWindow)
+	{
+		[FocusWindow makeKeyWindow];
+	}
 
 	return bSuccess;
 }
