Add a function to open a project file using the shell. Always opens as an Unreal.ProjectFile, preventing the user's default program selection from overriding it.

diff --git a/Engine/Source/Developer/DesktopPlatform/Private/DesktopPlatformBase.cpp b/Engine/Source/Developer/DesktopPlatform/Private/DesktopPlatformBase.cpp
--- a/Engine/Source/Developer/DesktopPlatform/Private/DesktopPlatformBase.cpp
+++ b/Engine/Source/Developer/DesktopPlatform/Private/DesktopPlatformBase.cpp
@@ -297,6 +297,12 @@ bool FDesktopPlatformBase::GetEngineIdentifierForProject(const FString &ProjectF
 	return false;
 }
 
+bool FDesktopPlatformBase::OpenProject(const FString& ProjectFileName)
+{
+	FPlatformProcess::LaunchFileInDefaultExternalApplication(*ProjectFileName);
+	return true;
+}
+
 bool FDesktopPlatformBase::CleanGameProject(const FString &ProjectDir, FFeedbackContext* Warn)
 {
 	// Begin a task
diff --git a/Engine/Source/Developer/DesktopPlatform/Private/DesktopPlatformBase.h b/Engine/Source/Developer/DesktopPlatform/Private/DesktopPlatformBase.h
--- a/Engine/Source/Developer/DesktopPlatform/Private/DesktopPlatformBase.h
+++ b/Engine/Source/Developer/DesktopPlatform/Private/DesktopPlatformBase.h
@@ -33,6 +33,8 @@ public:
 	virtual bool SetEngineIdentifierForProject(const FString &ProjectFileName, const FString &Identifier) override;
 	virtual bool GetEngineIdentifierForProject(const FString &ProjectFileName, FString &OutIdentifier) override;
 
+	virtual bool OpenProject(const FString& ProjectFileName) override;
+
 	virtual bool CleanGameProject(const FString& ProjectDir, FFeedbackContext* Warn) override;
 	virtual bool CompileGameProject(const FString& RootDir, const FString& ProjectFileName, FFeedbackContext* Warn) override;
 	virtual bool GenerateProjectFiles(const FString& RootDir, const FString& ProjectFileName, FFeedbackContext* Warn) override;
diff --git a/Engine/Source/Developer/DesktopPlatform/Private/Windows/DesktopPlatformWindows.cpp b/Engine/Source/Developer/DesktopPlatform/Private/Windows/DesktopPlatformWindows.cpp
--- a/Engine/Source/Developer/DesktopPlatform/Private/Windows/DesktopPlatformWindows.cpp
+++ b/Engine/Source/Developer/DesktopPlatform/Private/Windows/DesktopPlatformWindows.cpp
@@ -6,6 +6,7 @@
 
 #include "AllowWindowsPlatformTypes.h"
 	#include <commdlg.h>
+	#include <shellapi.h>
 	#include <shlobj.h>
 	#include <Winver.h>
 #include "HideWindowsPlatformTypes.h"
@@ -475,6 +476,20 @@ bool FDesktopPlatformWindows::UpdateFileAssociations()
 	return true;
 }
 
+bool FDesktopPlatformWindows::OpenProject(const FString &ProjectFileName)
+{
+	// Always treat projects as an Unreal.ProjectFile, don't allow the user overriding the file association to take effect
+	SHELLEXECUTEINFO Info;
+	ZeroMemory(&Info, sizeof(Info));
+	Info.cbSize = sizeof(Info);
+	Info.fMask = SEE_MASK_CLASSNAME;
+	Info.lpVerb = TEXT("open");
+	Info.nShow = SW_SHOWNORMAL;
+	Info.lpFile = *ProjectFileName;
+	Info.lpClass = TEXT("Unreal.ProjectFile");
+	return ::ShellExecuteExW(&Info) != 0;
+}
+
 bool FDesktopPlatformWindows::RunUnrealBuildTool(const FText& Description, const FString& RootDir, const FString& Arguments, FFeedbackContext* Warn)
 {
 	// Get the path to UBT
diff --git a/Engine/Source/Developer/DesktopPlatform/Private/Windows/DesktopPlatformWindows.h b/Engine/Source/Developer/DesktopPlatform/Private/Windows/DesktopPlatformWindows.h
--- a/Engine/Source/Developer/DesktopPlatform/Private/Windows/DesktopPlatformWindows.h
+++ b/Engine/Source/Developer/DesktopPlatform/Private/Windows/DesktopPlatformWindows.h
@@ -24,6 +24,8 @@ public:
 	virtual bool VerifyFileAssociations() override;
 	virtual bool UpdateFileAssociations() override;
 
+	virtual bool OpenProject(const FString &ProjectFileName) override;
+
 	virtual bool RunUnrealBuildTool(const FText& Description, const FString& RootDir, const FString& Arguments, FFeedbackContext* Warn) override;
 
 	virtual FFeedbackContext* GetNativeFeedbackContext() override;
diff --git a/Engine/Source/Developer/DesktopPlatform/Public/IDesktopPlatform.h b/Engine/Source/Developer/DesktopPlatform/Public/IDesktopPlatform.h
--- a/Engine/Source/Developer/DesktopPlatform/Public/IDesktopPlatform.h
+++ b/Engine/Source/Developer/DesktopPlatform/Public/IDesktopPlatform.h
@@ -286,6 +286,14 @@ public:
 	*/
 	virtual bool GetEngineIdentifierForProject(const FString &ProjectFileName, FString &OutIdentifier) = 0;
 
+	/**
+	* Opens the given project with the appropriate editor. Tries to use the shell association.
+	*
+	* @param ProjectFileName	Filename of the project to update
+	* @return true if the project was opened successfully.
+	*/
+	virtual bool OpenProject(const FString& ProjectFileName) = 0;
+
 	/**
 	* Cleans a game project. Removes the intermediate folder and binary build products.
 	*
