Normalize path for Win32 call to opening a project using the shell.

diff --git a/Engine/Source/Developer/DesktopPlatform/Private/Windows/DesktopPlatformWindows.cpp b/Engine/Source/Developer/DesktopPlatform/Private/Windows/DesktopPlatformWindows.cpp
--- a/Engine/Source/Developer/DesktopPlatform/Private/Windows/DesktopPlatformWindows.cpp
+++ b/Engine/Source/Developer/DesktopPlatform/Private/Windows/DesktopPlatformWindows.cpp
@@ -478,6 +478,10 @@ bool FDesktopPlatformWindows::UpdateFileAssociations()
 
 bool FDesktopPlatformWindows::OpenProject(const FString &ProjectFileName)
 {
+	// Get the project filename in a native format
+	FString PlatformProjectFileName = ProjectFileName;
+	FPaths::MakePlatformFilename(PlatformProjectFileName);
+
 	// Always treat projects as an Unreal.ProjectFile, don't allow the user overriding the file association to take effect
 	SHELLEXECUTEINFO Info;
 	ZeroMemory(&Info, sizeof(Info));
@@ -485,7 +489,7 @@ bool FDesktopPlatformWindows::OpenProject(const FString &ProjectFileName)
 	Info.fMask = SEE_MASK_CLASSNAME;
 	Info.lpVerb = TEXT("open");
 	Info.nShow = SW_SHOWNORMAL;
-	Info.lpFile = *ProjectFileName;
+	Info.lpFile = *PlatformProjectFileName;
 	Info.lpClass = TEXT("Unreal.ProjectFile");
 	return ::ShellExecuteExW(&Info) != 0;
 }
