Fix SlateViewer on Linux.

diff --git a/Engine/Source/Developer/DesktopPlatform/Private/Linux/DesktopPlatformLinux.cpp b/Engine/Source/Developer/DesktopPlatform/Private/Linux/DesktopPlatformLinux.cpp
--- a/Engine/Source/Developer/DesktopPlatform/Private/Linux/DesktopPlatformLinux.cpp
+++ b/Engine/Source/Developer/DesktopPlatform/Private/Linux/DesktopPlatformLinux.cpp
@@ -10,6 +10,12 @@ bool FDesktopPlatformLinux::OpenFileDialog(const void* ParentWindowHandle, const
 	return false;
 }
 
+bool FDesktopPlatformLinux::OpenFileDialog(const void* ParentWindowHandle, const FString& DialogTitle, const FString& DefaultPath, const FString& DefaultFile, const FString& FileTypes, uint32 Flags, TArray<FString>& OutFilenames, int32& OutFilterIndex)
+{
+	unimplemented();
+	return false;
+}
+
 bool FDesktopPlatformLinux::SaveFileDialog(const void* ParentWindowHandle, const FString& DialogTitle, const FString& DefaultPath, const FString& DefaultFile, const FString& FileTypes, uint32 Flags, TArray<FString>& OutFilenames)
 {
 	unimplemented();
diff --git a/Engine/Source/Developer/DesktopPlatform/Private/Linux/DesktopPlatformLinux.h b/Engine/Source/Developer/DesktopPlatform/Private/Linux/DesktopPlatformLinux.h
--- a/Engine/Source/Developer/DesktopPlatform/Private/Linux/DesktopPlatformLinux.h
+++ b/Engine/Source/Developer/DesktopPlatform/Private/Linux/DesktopPlatformLinux.h
@@ -9,6 +9,7 @@ class FDesktopPlatformLinux : public FDesktopPlatformBase
 public:
 	// IDesktopPlatform Implementation
 	virtual bool OpenFileDialog(const void* ParentWindowHandle, const FString& DialogTitle, const FString& DefaultPath, const FString& DefaultFile, const FString& FileTypes, uint32 Flags, TArray<FString>& OutFilenames) override;
+	virtual bool OpenFileDialog(const void* ParentWindowHandle, const FString& DialogTitle, const FString& DefaultPath, const FString& DefaultFile, const FString& FileTypes, uint32 Flags, TArray<FString>& OutFilenames, int32& OutFilterIndex) override;
 	virtual bool SaveFileDialog(const void* ParentWindowHandle, const FString& DialogTitle, const FString& DefaultPath, const FString& DefaultFile, const FString& FileTypes, uint32 Flags, TArray<FString>& OutFilenames) override;
 	virtual bool OpenDirectoryDialog(const void* ParentWindowHandle, const FString& DialogTitle, const FString& DefaultPath, FString& OutFolderName) override;
 	virtual bool OpenFontDialog(const void* ParentWindowHandle, FString& OutFontName, float& OutHeight, EFontImportFlags::Type& OutFlags) override;
