Implement a proper native feedback window for OS X that provides cancellation, progress status and a disclosable output log.

diff --git a/Engine/Source/Developer/DesktopPlatform/Private/Mac/MacNativeFeedbackContext.cpp b/Engine/Source/Developer/DesktopPlatform/Private/Mac/MacNativeFeedbackContext.cpp
--- a/Engine/Source/Developer/DesktopPlatform/Private/Mac/MacNativeFeedbackContext.cpp
+++ b/Engine/Source/Developer/DesktopPlatform/Private/Mac/MacNativeFeedbackContext.cpp
@@ -3,23 +3,360 @@
 #include "DesktopPlatformPrivatePCH.h"
 #include "MacNativeFeedbackContext.h"
 
+@implementation FMacNativeFeedbackContextWindowController
+
+-(void)toggleLog
+{
+	if([LogView isHidden])
+	{
+		NSRect Frame = [Window frame];
+		
+		int32 ConsoleHeight = 600;
+		if(GConfig)
+		{
+			GConfig->GetInt(TEXT("DebugMac"), TEXT("ConsoleHeight"), ConsoleHeight, GGameIni);
+		}
+		
+		Frame.origin.y -= (ConsoleHeight - Frame.size.height);
+		Frame.size.height = ConsoleHeight;
+
+		[Window setFrame:Frame display:YES animate:YES];
+		[LogView setHidden:NO];
+	}
+	else
+	{
+		CGFloat Height = [LogView frame].size.height;
+		[LogView setHidden:YES];
+		NSRect Frame = [Window frame];
+		Frame.size.height -= Height;
+		Frame.origin.y += Height;
+		[Window setFrame:Frame display:YES animate:YES];
+	}
+}
+
+-(id)init
+{
+	id obj = [super init];
+	if(obj)
+	{
+		int32 ConsoleWidth = 800;
+		int32 ConsoleHeight = 600;
+		int32 ConsolePosX = 0;
+		int32 ConsolePosY = 0;
+		bool bHasX = false;
+		bool bHasY = false;
+		
+		if(GConfig)
+		{
+			GConfig->GetInt(TEXT("DebugMac"), TEXT("ConsoleWidth"), ConsoleWidth, GGameIni);
+			GConfig->GetInt(TEXT("DebugMac"), TEXT("ConsoleHeight"), ConsoleHeight, GGameIni);
+			bHasX = GConfig->GetInt(TEXT("DebugMac"), TEXT("ConsoleX"), ConsolePosX, GGameIni);
+			bHasY = GConfig->GetInt(TEXT("DebugMac"), TEXT("ConsoleY"), ConsolePosY, GGameIni);
+		}
+		
+		NSRect WindowRect = NSMakeRect(ConsolePosX, ConsolePosY, ConsoleWidth, ConsoleHeight);
+		
+		Window = [[[NSWindow alloc] initWithContentRect:WindowRect styleMask:NSTitledWindowMask|NSMiniaturizableWindowMask|NSResizableWindowMask|NSClosableWindowMask backing:NSBackingStoreBuffered defer:NO] autorelease];
+		[Window setTitle:@"Unreal Engine 4"];
+		[Window setReleasedWhenClosed:NO];
+		[Window setMinSize:NSMakeSize(400, 100)];
+		
+		NSView* View = [Window contentView];
+		[View setAutoresizesSubviews:YES];
+		[View setAutoresizingMask:NSViewWidthSizable | NSViewHeightSizable];
+		
+		ShowLogButton = [[NSButton new] autorelease];
+		NSRect ShowLogRect = [ShowLogButton frame];
+		{
+			[ShowLogButton setIdentifier:@"ShowLogButton"];
+			[ShowLogButton setButtonType:NSMomentaryPushInButton];
+			[ShowLogButton setBezelStyle:NSRoundedBezelStyle];
+			[ShowLogButton setTitle:@"Show Log"];
+			[ShowLogButton sizeToFit];
+			ShowLogRect = [ShowLogButton frame];
+			ShowLogRect.origin.x = WindowRect.size.width - 8 - ShowLogRect.size.width;
+			ShowLogRect.origin.y = WindowRect.size.height - ShowLogRect.size.height - 8;
+			
+			[ShowLogButton setFrameOrigin:ShowLogRect.origin];
+			[ShowLogButton setAutoresizingMask:NSViewMinXMargin | NSViewMinYMargin];
+			[ShowLogButton setTarget:self];
+			[ShowLogButton setAction:@selector(toggleLog)];
+		}
+		
+		CancelButton = [[NSButton new] autorelease];
+		NSRect CancelRect = [CancelButton frame];
+		{
+			[CancelButton setIdentifier:@"CancelButton"];
+			[CancelButton setTitle:@"Cancel"];
+			[CancelButton setButtonType:NSMomentaryPushInButton];
+			[CancelButton setBezelStyle:NSRoundedBezelStyle];
+			[CancelButton sizeToFit];
+			CancelRect = [CancelButton frame];
+			CancelRect.origin.x = ShowLogRect.origin.x - CancelRect.size.width - 4;
+			CancelRect.origin.y = ShowLogRect.origin.y;
+			
+			[CancelButton setFrameOrigin:CancelRect.origin];
+			[CancelButton setAutoresizingMask:NSViewMinXMargin | NSViewMinYMargin];
+			[CancelButton setTarget:self];
+			[CancelButton setAction:@selector(hideWindow)];
+		}
+		
+		StatusLabel = [[[NSTextField alloc] initWithFrame:NSMakeRect(0, 0, 100, 18)] autorelease];
+		NSRect StatusRect = [StatusLabel frame];
+		{
+			[StatusLabel setBezeled:NO];
+			[StatusLabel setDrawsBackground:NO];
+			[StatusLabel setFont:[NSFont labelFontOfSize:[NSFont systemFontSize]]];
+			[StatusLabel setSelectable:NO];
+			[StatusLabel setEditable:NO];
+			[StatusLabel setBordered:NO];
+			[StatusLabel setStringValue:@"Progress:"];
+			StatusRect = [StatusLabel frame];
+			
+			StatusRect.size.width = CancelRect.origin.x - 16;
+			StatusRect.origin.x = 8;
+			StatusRect.origin.y = ShowLogRect.origin.y + ((ShowLogRect.size.height - [NSFont systemFontSize]) / 2);
+			
+			[StatusLabel setIdentifier:@"StatusLabel"];
+			[StatusLabel setFrame:StatusRect];
+			[StatusLabel setAutoresizingMask:NSViewWidthSizable | NSViewMinYMargin];
+		}
+		
+		ProgressIndicator = [[NSProgressIndicator new] autorelease];
+		NSRect ProgressRect = [ProgressIndicator frame];
+		{
+			[ProgressIndicator setStyle:NSProgressIndicatorBarStyle];
+			[ProgressIndicator sizeToFit];
+			ProgressRect = [ProgressIndicator frame];
+			ProgressRect.size.width = WindowRect.size.width - 16;
+			ProgressRect.origin.x = 8;
+			ProgressRect.origin.y = CancelRect.origin.y - ProgressRect.size.height - 8;
+			
+			[ProgressIndicator setIdentifier:@"ProgressIndicator"];
+			[ProgressIndicator setIndeterminate:YES];
+			[ProgressIndicator setFrame:ProgressRect];
+			[ProgressIndicator setAutoresizingMask:NSViewWidthSizable | NSViewMinYMargin];
+		}
+		
+		TextView = [[NSTextView new] autorelease];
+		NSRect TextRect = [TextView frame];
+		{
+			[TextView setIdentifier:@"TextView"];
+			[TextView setVerticallyResizable:YES];
+			[TextView setHorizontallyResizable:NO];
+			[TextView setBackgroundColor: [NSColor blackColor]];
+			[TextView setMinSize:NSMakeSize( 0.0, 0.0 ) ];
+			[TextView setMaxSize:NSMakeSize( FLT_MAX, FLT_MAX )];
+			TextRect = NSMakeRect(8, 8, WindowRect.size.width - 16, ProgressRect.origin.y - 16);
+			
+			[TextView setAutoresizingMask:NSViewWidthSizable | NSViewHeightSizable];
+			
+			LogView = [[NSScrollView new] autorelease];
+			[LogView setHasVerticalScroller:YES];
+			[LogView setHasHorizontalScroller:NO];
+			[LogView setAutohidesScrollers:YES];
+			[LogView setAutoresizesSubviews:YES];
+			[LogView setAutoresizingMask:NSViewWidthSizable | NSViewHeightSizable | NSViewMaxYMargin];
+			[LogView setFrame:TextRect];
+			
+			[TextView setFrameSize:[LogView contentSize]];
+			[[TextView textContainer] setContainerSize:NSMakeSize( [LogView contentSize].width, FLT_MAX )];
+			[[TextView textContainer] setWidthTracksTextView:YES];
+			
+			[LogView setDocumentView:TextView];
+		}
+		
+		[View addSubview:ShowLogButton];
+		[View addSubview:StatusLabel];
+		[View addSubview:ProgressIndicator];
+		[View addSubview:CancelButton];
+		[View addSubview:LogView];
+		
+		[ProgressIndicator startAnimation:nil];
+		
+		if (!bHasX || !bHasY)
+		{
+			[Window center];
+		}
+	}
+	return self;
+}
+
+-(void)dealloc
+{
+	[Window close];
+	[Window release];
+	[super dealloc];
+}
+
+-(void)setShowProgress:(bool)bShowProgress
+{
+	if(bShowProgress)
+	{
+		[ProgressIndicator stopAnimation:nil];
+		[ProgressIndicator setIndeterminate:NO];
+	}
+	else
+	{
+		if(![LogView isHidden])
+		{
+			[self toggleLog];
+		}
+		[ProgressIndicator setIndeterminate:YES];
+		[ProgressIndicator startAnimation:nil];
+	}
+}
+
+-(void)setShowCancelButton:(bool)bShowCancelButton
+{
+	BOOL bHideCancel = !bShowCancelButton;
+	[CancelButton setHidden:bHideCancel];
+}
+
+-(void)setTitleText:(NSString*)Title
+{
+	[Window setTitle:Title];
+}
+
+-(void)setStatusText:(NSString*)Text
+{
+	[StatusLabel setStringValue:Text];
+}
+
+-(void)setProgress:(double)Progress total:(double)Total
+{
+	if(![ProgressIndicator isIndeterminate])
+	{
+		[ProgressIndicator setMaxValue:Total];
+		[ProgressIndicator setMinValue:0.0];
+		[ProgressIndicator setDoubleValue:Progress];
+	}
+}
+
+-(void)showWindow
+{
+	[Window makeKeyAndOrderFront:nil];
+}
+
+-(void)hideWindow
+{
+	[Window orderOut:nil];
+}
+
+-(bool)windowOpen
+{
+	return [Window isVisible];
+}
+@end
+
 FMacNativeFeedbackContext::FMacNativeFeedbackContext()
 	: FFeedbackContext()
+	, WindowController(nil)
 	, Context( NULL )
 	, SlowTaskCount( 0 )
 	, bShowingConsoleForSlowTask( false )
 {
+	WindowController = [[FMacNativeFeedbackContextWindowController alloc] init];
+	
+	SetDefaultTextColor();
 }
 
 FMacNativeFeedbackContext::~FMacNativeFeedbackContext()
 {
+	[WindowController release];
 }
 
-void FMacNativeFeedbackContext::Serialize( const TCHAR* V, ELogVerbosity::Type Verbosity, const class FName& Category )
+void FMacNativeFeedbackContext::Serialize( const TCHAR* Data, ELogVerbosity::Type Verbosity, const class FName& Category )
 {
 	if( !GLog->IsRedirectingTo( this ) )
 	{
-		GLog->Serialize( V, Verbosity, Category );
+		GLog->Serialize( Data, Verbosity, Category );
+	}
+	
+	if(WindowController != NULL && bShowingConsoleForSlowTask)
+	{
+		FScopeLock ScopeLock( &CriticalSection );
+		
+		static bool Entry=0;
+		if( !GIsCriticalError || Entry )
+		{
+			// here we can change the color of the text to display, it's in the format:
+			// ForegroundRed | ForegroundGreen | ForegroundBlue | ForegroundBright | BackgroundRed | BackgroundGreen | BackgroundBlue | BackgroundBright
+			// where each value is either 0 or 1 (can leave off trailing 0's), so
+			// blue on bright yellow is "00101101" and red on black is "1"
+			// An empty string reverts to the normal gray on black
+			
+			if (Verbosity == ELogVerbosity::SetColor)
+			{
+				if (FCString::Stricmp(Data, TEXT("")) == 0)
+				{
+					SetDefaultTextColor();
+				}
+				else
+				{
+					SCOPED_AUTORELEASE_POOL;
+					
+					// turn the string into a bunch of 0's and 1's
+					TCHAR String[9];
+					FMemory::Memset(String, 0, sizeof(TCHAR) * ARRAY_COUNT(String));
+					FCString::Strncpy(String, Data, ARRAY_COUNT(String));
+					for (TCHAR* S = String; *S; S++)
+					{
+						*S -= '0';
+					}
+					
+					NSColor* Colors[2];
+					NSString* AttributeKeys[2];
+					
+					// Get FOREGROUND_INTENSITY and calculate final color
+					CGFloat Intensity = String[3] ? 1.0 : 0.5;
+					Colors[0] = [NSColor colorWithSRGBRed:(String[0] ? 1.0 * Intensity : 0.0) green:(String[1] ? 1.0 * Intensity : 0.0) blue:(String[2] ? 1.0 * Intensity : 0.0) alpha:1.0];
+					
+					// Get BACKGROUND_INTENSITY and calculate final color
+					Intensity = String[7] ? 1.0 : 0.5;
+					Colors[1] = [NSColor colorWithSRGBRed:(String[4] ? 1.0 * Intensity : 0.0) green:(String[5] ? 1.0 * Intensity : 0.0) blue:(String[6] ? 1.0 * Intensity : 0.0) alpha:1.0];
+					
+					AttributeKeys[0] = NSForegroundColorAttributeName;
+					AttributeKeys[1] = NSBackgroundColorAttributeName;
+					
+					if( TextViewTextColor )
+					[TextViewTextColor release];
+					
+					TextViewTextColor = [[NSDictionary alloc] initWithObjects:(id *)Colors forKeys:(id *)AttributeKeys count:2];
+				}
+			}
+			else
+			{
+				SCOPED_AUTORELEASE_POOL;
+				
+				TCHAR OutputString[MAX_SPRINTF]=TEXT(""); //@warning: this is safe as FCString::Sprintf only use 1024 characters max
+				FCString::Sprintf(OutputString,TEXT("%s%s"),*FOutputDevice::FormatLogLine(Verbosity, Category, Data, GPrintLogTimes),LINE_TERMINATOR);
+				
+				CFStringRef CocoaText = FPlatformString::TCHARToCFString(OutputString);
+				NSAttributedString *AttributedString = [[NSAttributedString alloc] initWithString:(NSString*)CocoaText attributes:TextViewTextColor];
+				[[WindowController->TextView textStorage] appendAttributedString:AttributedString];
+				[WindowController->TextView scrollRangeToVisible:NSMakeRange([[WindowController->TextView string] length], 0)];
+				CFRelease(CocoaText);
+				[AttributedString release];
+				
+				if(!MacApplication)
+				{
+					FPlatformMisc::PumpMessages( true );
+				}
+			}
+		}
+		else
+		{
+			Entry=1;
+			try
+			{
+				// Ignore errors to prevent infinite-recursive exception reporting.
+				Serialize( Data, Verbosity, Category );
+			}
+			catch( ... )
+			{}
+			Entry=0;
+		}
 	}
 }
 
@@ -30,16 +367,31 @@ bool FMacNativeFeedbackContext::YesNof(const FText& Text)
 
 bool FMacNativeFeedbackContext::ReceivedUserCancel()
 {
-	return false;
+	bool bReceivedUserCancel = false;
+	if(WindowController != NULL && bShowingConsoleForSlowTask && ![WindowController windowOpen])
+	{
+		bReceivedUserCancel = true;
+	}
+	return bReceivedUserCancel;
 }
 
 void FMacNativeFeedbackContext::BeginSlowTask( const FText& Task, bool bShowProgressDialog, bool bInShowCancelButton )
 {
 	GIsSlowTask = ++SlowTaskCount>0;
 
-	if(SlowTaskCount > 0 && bShowProgressDialog && GLogConsole != NULL && !GLogConsole->IsShown())
+	if(SlowTaskCount > 0 && bShowProgressDialog && WindowController != NULL && !bShowingConsoleForSlowTask)
 	{
-		GLogConsole->Show(true);
+		[WindowController setTitleText:Task.ToString().GetNSString()];
+		[WindowController setStatusText:@"Progress:"];
+		[WindowController setShowCancelButton:bInShowCancelButton];
+		[WindowController setShowProgress:true];
+		[WindowController setProgress:0 total:100];
+		[WindowController setShowProgress:false];
+		
+		SetDefaultTextColor();
+		
+		[WindowController showWindow];
+		
 		bShowingConsoleForSlowTask = true;
 	}
 }
@@ -51,23 +403,33 @@ void FMacNativeFeedbackContext::EndSlowTask()
 
 	if(SlowTaskCount == 0 && bShowingConsoleForSlowTask)
 	{
-		if(GLogConsole != NULL) GLogConsole->Show(false);
+		if(WindowController != NULL) [WindowController hideWindow];
 		bShowingConsoleForSlowTask = false;
 	}
 }
 
 bool FMacNativeFeedbackContext::StatusUpdate( int32 Numerator, int32 Denominator, const FText& NewStatus )
 {
+	if(WindowController != NULL && bShowingConsoleForSlowTask)
+	{
+		[WindowController setStatusText:NewStatus.ToString().GetNSString()];
+		UpdateProgress(Numerator, Denominator);
+	}
 	return true;
 }
 
 bool FMacNativeFeedbackContext::StatusForceUpdate( int32 Numerator, int32 Denominator, const FText& StatusText )
 {
-	return true;
+	return StatusUpdate(Numerator, Denominator, StatusText);
 }
 
 void FMacNativeFeedbackContext::UpdateProgress(int32 Numerator, int32 Denominator)
 {
+	if(WindowController != NULL && bShowingConsoleForSlowTask)
+	{
+		[WindowController setShowProgress:true];
+		[WindowController setProgress:Numerator total:Denominator];
+	}
 }
 
 FContextSupplier* FMacNativeFeedbackContext::GetContext() const
@@ -79,3 +441,22 @@ void FMacNativeFeedbackContext::SetContext( FContextSupplier* InSupplier )
 {
 	Context = InSupplier;
 }
+
+void FMacNativeFeedbackContext::SetDefaultTextColor()
+{
+	SCOPED_AUTORELEASE_POOL;
+	
+	NSColor* Colors[2];
+	NSString* AttributeKeys[2];
+	
+	Colors[0] = [NSColor grayColor];
+	Colors[1] = [NSColor blackColor];
+	
+	AttributeKeys[0] = NSForegroundColorAttributeName;
+	AttributeKeys[1] = NSBackgroundColorAttributeName;
+	
+	if( TextViewTextColor )
+	[TextViewTextColor release];
+	
+	TextViewTextColor = [[NSDictionary alloc] initWithObjects:(id *)Colors forKeys:(id *)AttributeKeys count:2];
+}
diff --git a/Engine/Source/Developer/DesktopPlatform/Private/Mac/MacNativeFeedbackContext.h b/Engine/Source/Developer/DesktopPlatform/Private/Mac/MacNativeFeedbackContext.h
--- a/Engine/Source/Developer/DesktopPlatform/Private/Mac/MacNativeFeedbackContext.h
+++ b/Engine/Source/Developer/DesktopPlatform/Private/Mac/MacNativeFeedbackContext.h
@@ -6,6 +6,30 @@
 
 #pragma once
 
+@interface FMacNativeFeedbackContextWindowController : NSObject
+{
+@public
+	NSTextView* TextView;
+@private
+	NSScrollView* LogView;
+	NSWindow* Window;
+	NSTextField* StatusLabel;
+	NSButton* CancelButton;
+	NSButton* ShowLogButton;
+	NSProgressIndicator* ProgressIndicator;
+}
+-(id)init;
+-(void)dealloc;
+-(void)setShowProgress:(bool)bShowProgress;
+-(void)setShowCancelButton:(bool)bShowCancelButton;
+-(void)setTitleText:(NSString*)Title;
+-(void)setStatusText:(NSString*)Text;
+-(void)setProgress:(double)Progress total:(double)Total;
+-(void)showWindow;
+-(void)hideWindow;
+-(bool)windowOpen;
+@end
+
 /**
  * Feedback context implementation for Mac.
  */
@@ -32,6 +56,13 @@ public:
 	void SetContext( FContextSupplier* InSupplier );
 
 private:
+	void SetDefaultTextColor();
+	
+private:
+	/** Critical section for Serialize() */
+	FCriticalSection CriticalSection;
+	FMacNativeFeedbackContextWindowController* WindowController;
+	NSDictionary* TextViewTextColor;
 	FContextSupplier* Context;
 	int32 SlowTaskCount;
 	bool bShowingConsoleForSlowTask;
