Renamed the folders used to store the Editor/Launcher Platform icons so that the distill script doesn't strip them out

TTP# 340470 - BINARY: EDITOR: Platform icons are missing

The distill script assumes that anything inside a folder named the same as a platform name (eg, PS4) is a platform specific thing (which these icons are not). I've renamed the folders from "Platform" (eg, PS4) to "PlatformTarget" (eg, PS4Target) and updated the PlatformInfo array to reference their new location.

diff --git a/Engine/Source/Developer/DesktopPlatform/Private/PlatformInfo.cpp b/Engine/Source/Developer/DesktopPlatform/Private/PlatformInfo.cpp
--- a/Engine/Source/Developer/DesktopPlatform/Private/PlatformInfo.cpp
+++ b/Engine/Source/Developer/DesktopPlatform/Private/PlatformInfo.cpp
@@ -48,41 +48,41 @@ FPlatformInfo BuildPlatformInfo(const FName& InPlatformInfoName, const FName& In
 }
 
 static const FPlatformInfo PlatformInfoArray[] = {
-	//				  PlatformInfoName					TargetPlatformName			DisplayName													PlatformType				PlatformFlags					IconPaths																																UATCommandLine
-	BuildPlatformInfo(TEXT("WindowsNoEditor"),			TEXT("WindowsNoEditor"),	LOCTEXT("WindowsNoEditor", "Windows"),						EPlatformType::Game,		EPlatformFlags::None,			FPlatformIconPaths(TEXT("Launcher/Windows/Platform_WindowsNoEditor_24x"), TEXT("Launcher/Windows/Platform_WindowsNoEditor_128x")),		TEXT("")),
-	BuildPlatformInfo(TEXT("WindowsNoEditor_Win32"),	TEXT("WindowsNoEditor"),	LOCTEXT("WindowsNoEditor_Win32", "Windows (32-bit)"),		EPlatformType::Game,		EPlatformFlags::BuildFlavor,	FPlatformIconPaths(TEXT("Launcher/Windows/Platform_WindowsNoEditor_24x"), TEXT("Launcher/Windows/Platform_WindowsNoEditor_128x")),		TEXT("-targetplatform=Win32")),
-	BuildPlatformInfo(TEXT("WindowsNoEditor_Win64"),	TEXT("WindowsNoEditor"),	LOCTEXT("WindowsNoEditor_Win64", "Windows (64-bit)"),		EPlatformType::Game,		EPlatformFlags::BuildFlavor,	FPlatformIconPaths(TEXT("Launcher/Windows/Platform_WindowsNoEditor_24x"), TEXT("Launcher/Windows/Platform_WindowsNoEditor_128x")),		TEXT("-targetplatform=Win64")),
-	BuildPlatformInfo(TEXT("Windows"),					TEXT("Windows"),			LOCTEXT("WindowsEditor", "Windows (Editor)"),				EPlatformType::Editor,		EPlatformFlags::None,			FPlatformIconPaths(TEXT("Launcher/Windows/Platform_Windows_24x"), TEXT("Launcher/Windows/Platform_Windows_128x")),						TEXT("")),
-	BuildPlatformInfo(TEXT("WindowsClient"),			TEXT("WindowsClient"),		LOCTEXT("WindowsClient", "Windows (Client-only)"),			EPlatformType::Client,		EPlatformFlags::None,			FPlatformIconPaths(TEXT("Launcher/Windows/Platform_Windows_24x"), TEXT("Launcher/Windows/Platform_Windows_128x")),						TEXT("")),
-	BuildPlatformInfo(TEXT("WindowsServer"),			TEXT("WindowsServer"),		LOCTEXT("WindowsServer", "Windows (Dedicated Server)"),		EPlatformType::Server,		EPlatformFlags::None,			FPlatformIconPaths(TEXT("Launcher/Windows/Platform_WindowsServer_24x"), TEXT("Launcher/Windows/Platform_WindowsServer_128x")),			TEXT("")),
-
-	BuildPlatformInfo(TEXT("MacNoEditor"),				TEXT("MacNoEditor"),		LOCTEXT("MacNoEditor", "Mac"),								EPlatformType::Game,		EPlatformFlags::None,			FPlatformIconPaths(TEXT("Launcher/Mac/Platform_Mac_24x"), TEXT("Launcher/Mac/Platform_Mac_128x")),										TEXT("")),
-	BuildPlatformInfo(TEXT("Mac"),						TEXT("Mac"),				LOCTEXT("MacEditor", "Mac (Editor)"),						EPlatformType::Editor,		EPlatformFlags::None,			FPlatformIconPaths(TEXT("Launcher/Mac/Platform_Mac_24x"), TEXT("Launcher/Mac/Platform_Mac_128x")),										TEXT("")),
-	BuildPlatformInfo(TEXT("MacClient"),				TEXT("MacClient"),			LOCTEXT("MacClient", "Mac (Client-only)"),					EPlatformType::Client,		EPlatformFlags::None,			FPlatformIconPaths(TEXT("Launcher/Mac/Platform_Mac_24x"), TEXT("Launcher/Mac/Platform_Mac_128x")),										TEXT("")),
-	BuildPlatformInfo(TEXT("MacServer"),				TEXT("MacServer"),			LOCTEXT("MacServer", "Mac (Dedicated Server)"),				EPlatformType::Server,		EPlatformFlags::None,			FPlatformIconPaths(TEXT("Launcher/Mac/Platform_Mac_24x"), TEXT("Launcher/Mac/Platform_Mac_128x")),										TEXT("")),
-
-	BuildPlatformInfo(TEXT("LinuxNoEditor"),			TEXT("LinuxNoEditor"),		LOCTEXT("LinuxNoEditor", "Linux"),							EPlatformType::Game,		EPlatformFlags::None,			FPlatformIconPaths(TEXT("Launcher/Linux/Platform_Linux_24x"), TEXT("Launcher/Linux/Platform_Linux_128x")),								TEXT("")),
-	BuildPlatformInfo(TEXT("Linux"),					TEXT("Linux"),				LOCTEXT("LinuxEditor", "Linux (Editor)"),					EPlatformType::Editor,		EPlatformFlags::None,			FPlatformIconPaths(TEXT("Launcher/Linux/Platform_Linux_24x"), TEXT("Launcher/Linux/Platform_Linux_128x")),								TEXT("")),
-	BuildPlatformInfo(TEXT("LinuxClient"),				TEXT("LinuxClient"),		LOCTEXT("LinuxClient", "Linux (Client-only)"),				EPlatformType::Client,		EPlatformFlags::None,			FPlatformIconPaths(TEXT("Launcher/Linux/Platform_Linux_24x"), TEXT("Launcher/Linux/Platform_Linux_128x")),								TEXT("")),
-	BuildPlatformInfo(TEXT("LinuxServer"),				TEXT("LinuxServer"),		LOCTEXT("LinuxServer", "Linux (Dedicated Server)"),			EPlatformType::Server,		EPlatformFlags::None,			FPlatformIconPaths(TEXT("Launcher/Linux/Platform_Linux_24x"), TEXT("Launcher/Linux/Platform_Linux_128x")),								TEXT("")),
+	//				  PlatformInfoName					TargetPlatformName			DisplayName													PlatformType				PlatformFlags					IconPaths																																			UATCommandLine
+	BuildPlatformInfo(TEXT("WindowsNoEditor"),			TEXT("WindowsNoEditor"),	LOCTEXT("WindowsNoEditor", "Windows"),						EPlatformType::Game,		EPlatformFlags::None,			FPlatformIconPaths(TEXT("Launcher/WindowsTarget/Platform_WindowsNoEditor_24x"), TEXT("Launcher/WindowsTarget/Platform_WindowsNoEditor_128x")),		TEXT("")),
+	BuildPlatformInfo(TEXT("WindowsNoEditor_Win32"),	TEXT("WindowsNoEditor"),	LOCTEXT("WindowsNoEditor_Win32", "Windows (32-bit)"),		EPlatformType::Game,		EPlatformFlags::BuildFlavor,	FPlatformIconPaths(TEXT("Launcher/WindowsTarget/Platform_WindowsNoEditor_24x"), TEXT("Launcher/WindowsTarget/Platform_WindowsNoEditor_128x")),		TEXT("-targetplatform=Win32")),
+	BuildPlatformInfo(TEXT("WindowsNoEditor_Win64"),	TEXT("WindowsNoEditor"),	LOCTEXT("WindowsNoEditor_Win64", "Windows (64-bit)"),		EPlatformType::Game,		EPlatformFlags::BuildFlavor,	FPlatformIconPaths(TEXT("Launcher/WindowsTarget/Platform_WindowsNoEditor_24x"), TEXT("Launcher/WindowsTarget/Platform_WindowsNoEditor_128x")),		TEXT("-targetplatform=Win64")),
+	BuildPlatformInfo(TEXT("Windows"),					TEXT("Windows"),			LOCTEXT("WindowsEditor", "Windows (Editor)"),				EPlatformType::Editor,		EPlatformFlags::None,			FPlatformIconPaths(TEXT("Launcher/WindowsTarget/Platform_Windows_24x"), TEXT("Launcher/WindowsTarget/Platform_Windows_128x")),						TEXT("")),
+	BuildPlatformInfo(TEXT("WindowsClient"),			TEXT("WindowsClient"),		LOCTEXT("WindowsClient", "Windows (Client-only)"),			EPlatformType::Client,		EPlatformFlags::None,			FPlatformIconPaths(TEXT("Launcher/WindowsTarget/Platform_Windows_24x"), TEXT("Launcher/WindowsTarget/Platform_Windows_128x")),						TEXT("")),
+	BuildPlatformInfo(TEXT("WindowsServer"),			TEXT("WindowsServer"),		LOCTEXT("WindowsServer", "Windows (Dedicated Server)"),		EPlatformType::Server,		EPlatformFlags::None,			FPlatformIconPaths(TEXT("Launcher/WindowsTarget/Platform_WindowsServer_24x"), TEXT("Launcher/WindowsTarget/Platform_WindowsServer_128x")),			TEXT("")),
+
+	BuildPlatformInfo(TEXT("MacNoEditor"),				TEXT("MacNoEditor"),		LOCTEXT("MacNoEditor", "Mac"),								EPlatformType::Game,		EPlatformFlags::None,			FPlatformIconPaths(TEXT("Launcher/MacTarget/Platform_Mac_24x"), TEXT("Launcher/MacTarget/Platform_Mac_128x")),										TEXT("")),
+	BuildPlatformInfo(TEXT("Mac"),						TEXT("Mac"),				LOCTEXT("MacEditor", "Mac (Editor)"),						EPlatformType::Editor,		EPlatformFlags::None,			FPlatformIconPaths(TEXT("Launcher/MacTarget/Platform_Mac_24x"), TEXT("Launcher/MacTarget/Platform_Mac_128x")),										TEXT("")),
+	BuildPlatformInfo(TEXT("MacClient"),				TEXT("MacClient"),			LOCTEXT("MacClient", "Mac (Client-only)"),					EPlatformType::Client,		EPlatformFlags::None,			FPlatformIconPaths(TEXT("Launcher/MacTarget/Platform_Mac_24x"), TEXT("Launcher/MacTarget/Platform_Mac_128x")),										TEXT("")),
+	BuildPlatformInfo(TEXT("MacServer"),				TEXT("MacServer"),			LOCTEXT("MacServer", "Mac (Dedicated Server)"),				EPlatformType::Server,		EPlatformFlags::None,			FPlatformIconPaths(TEXT("Launcher/MacTarget/Platform_Mac_24x"), TEXT("Launcher/MacTarget/Platform_Mac_128x")),										TEXT("")),
+
+	BuildPlatformInfo(TEXT("LinuxNoEditor"),			TEXT("LinuxNoEditor"),		LOCTEXT("LinuxNoEditor", "Linux"),							EPlatformType::Game,		EPlatformFlags::None,			FPlatformIconPaths(TEXT("Launcher/LinuxTarget/Platform_Linux_24x"), TEXT("Launcher/LinuxTarget/Platform_Linux_128x")),								TEXT("")),
+	BuildPlatformInfo(TEXT("Linux"),					TEXT("Linux"),				LOCTEXT("LinuxEditor", "Linux (Editor)"),					EPlatformType::Editor,		EPlatformFlags::None,			FPlatformIconPaths(TEXT("Launcher/LinuxTarget/Platform_Linux_24x"), TEXT("Launcher/LinuxTarget/Platform_Linux_128x")),								TEXT("")),
+	BuildPlatformInfo(TEXT("LinuxClient"),				TEXT("LinuxClient"),		LOCTEXT("LinuxClient", "Linux (Client-only)"),				EPlatformType::Client,		EPlatformFlags::None,			FPlatformIconPaths(TEXT("Launcher/LinuxTarget/Platform_Linux_24x"), TEXT("Launcher/LinuxTarget/Platform_Linux_128x")),								TEXT("")),
+	BuildPlatformInfo(TEXT("LinuxServer"),				TEXT("LinuxServer"),		LOCTEXT("LinuxServer", "Linux (Dedicated Server)"),			EPlatformType::Server,		EPlatformFlags::None,			FPlatformIconPaths(TEXT("Launcher/LinuxTarget/Platform_Linux_24x"), TEXT("Launcher/LinuxTarget/Platform_Linux_128x")),								TEXT("")),
 	
-	BuildPlatformInfo(TEXT("Android"),					TEXT("Android"),			LOCTEXT("Android", "Android"),								EPlatformType::Game,		EPlatformFlags::None,			FPlatformIconPaths(TEXT("Launcher/Android/Platform_Android_24x"), TEXT("Launcher/Android/Platform_Android_128x")),						TEXT("")),
-	BuildPlatformInfo(TEXT("Android_All"),				TEXT("Android"),			LOCTEXT("Android_All", "Android (All)"),					EPlatformType::Game,		EPlatformFlags::CookFlavor,		FPlatformIconPaths(TEXT("Launcher/Android/Platform_Android_24x"), TEXT("Launcher/Android/Platform_Android_128x")),						TEXT("")),
-	BuildPlatformInfo(TEXT("Android_ATC"),				TEXT("Android_ATC"),		LOCTEXT("Android_ATC", "Android (ATC)"),					EPlatformType::Game,		EPlatformFlags::CookFlavor,		FPlatformIconPaths(TEXT("Launcher/Android/Platform_Android_ATC_24x"), TEXT("Launcher/Android/Platform_Android_128x")),					TEXT("-targetplatform=Android -cookflavor=ATC")),
-	BuildPlatformInfo(TEXT("Android_DXT"),				TEXT("Android_DXT"),		LOCTEXT("Android_DXT", "Android (DXT)"),					EPlatformType::Game,		EPlatformFlags::CookFlavor,		FPlatformIconPaths(TEXT("Launcher/Android/Platform_Android_DXT_24x"), TEXT("Launcher/Android/Platform_Android_128x")),					TEXT("-targetplatform=Android -cookflavor=DXT")),
-	BuildPlatformInfo(TEXT("Android_ETC1"),				TEXT("Android_ETC1"),		LOCTEXT("Android_ETC1", "Android (ETC1)"),					EPlatformType::Game,		EPlatformFlags::CookFlavor,		FPlatformIconPaths(TEXT("Launcher/Android/Platform_Android_ETC1_24x"), TEXT("Launcher/Android/Platform_Android_128x")),					TEXT("-targetplatform=Android -cookflavor=ETC1")),
-	BuildPlatformInfo(TEXT("Android_ETC2"),				TEXT("Android_ETC2"),		LOCTEXT("Android_ETC2", "Android (ETC2)"),					EPlatformType::Game,		EPlatformFlags::CookFlavor,		FPlatformIconPaths(TEXT("Launcher/Android/Platform_Android_ETC2_24x"), TEXT("Launcher/Android/Platform_Android_128x")),					TEXT("-targetplatform=Android -cookflavor=ETC2")),
-	BuildPlatformInfo(TEXT("Android_PVRTC"),			TEXT("Android_PVRTC"),		LOCTEXT("Android_PVRTC", "Android (PVRTC)"),				EPlatformType::Game,		EPlatformFlags::CookFlavor,		FPlatformIconPaths(TEXT("Launcher/Android/Platform_Android_PVRTC_24x"), TEXT("Launcher/Android/Platform_Android_128x")),				TEXT("-targetplatform=Android -cookflavor=PVRTC")),
+	BuildPlatformInfo(TEXT("Android"),					TEXT("Android"),			LOCTEXT("Android", "Android"),								EPlatformType::Game,		EPlatformFlags::None,			FPlatformIconPaths(TEXT("Launcher/AndroidTarget/Platform_Android_24x"), TEXT("Launcher/AndroidTarget/Platform_Android_128x")),						TEXT("")),
+	BuildPlatformInfo(TEXT("Android_All"),				TEXT("Android"),			LOCTEXT("Android_All", "Android (All)"),					EPlatformType::Game,		EPlatformFlags::CookFlavor,		FPlatformIconPaths(TEXT("Launcher/AndroidTarget/Platform_Android_24x"), TEXT("Launcher/AndroidTarget/Platform_Android_128x")),						TEXT("")),
+	BuildPlatformInfo(TEXT("Android_ATC"),				TEXT("Android_ATC"),		LOCTEXT("Android_ATC", "Android (ATC)"),					EPlatformType::Game,		EPlatformFlags::CookFlavor,		FPlatformIconPaths(TEXT("Launcher/AndroidTarget/Platform_Android_ATC_24x"), TEXT("Launcher/AndroidTarget/Platform_Android_128x")),					TEXT("-targetplatform=Android -cookflavor=ATC")),
+	BuildPlatformInfo(TEXT("Android_DXT"),				TEXT("Android_DXT"),		LOCTEXT("Android_DXT", "Android (DXT)"),					EPlatformType::Game,		EPlatformFlags::CookFlavor,		FPlatformIconPaths(TEXT("Launcher/AndroidTarget/Platform_Android_DXT_24x"), TEXT("Launcher/AndroidTarget/Platform_Android_128x")),					TEXT("-targetplatform=Android -cookflavor=DXT")),
+	BuildPlatformInfo(TEXT("Android_ETC1"),				TEXT("Android_ETC1"),		LOCTEXT("Android_ETC1", "Android (ETC1)"),					EPlatformType::Game,		EPlatformFlags::CookFlavor,		FPlatformIconPaths(TEXT("Launcher/AndroidTarget/Platform_Android_ETC1_24x"), TEXT("Launcher/AndroidTarget/Platform_Android_128x")),					TEXT("-targetplatform=Android -cookflavor=ETC1")),
+	BuildPlatformInfo(TEXT("Android_ETC2"),				TEXT("Android_ETC2"),		LOCTEXT("Android_ETC2", "Android (ETC2)"),					EPlatformType::Game,		EPlatformFlags::CookFlavor,		FPlatformIconPaths(TEXT("Launcher/AndroidTarget/Platform_Android_ETC2_24x"), TEXT("Launcher/AndroidTarget/Platform_Android_128x")),					TEXT("-targetplatform=Android -cookflavor=ETC2")),
+	BuildPlatformInfo(TEXT("Android_PVRTC"),			TEXT("Android_PVRTC"),		LOCTEXT("Android_PVRTC", "Android (PVRTC)"),				EPlatformType::Game,		EPlatformFlags::CookFlavor,		FPlatformIconPaths(TEXT("Launcher/AndroidTarget/Platform_Android_PVRTC_24x"), TEXT("Launcher/AndroidTarget/Platform_Android_128x")),				TEXT("-targetplatform=Android -cookflavor=PVRTC")),
 
-	BuildPlatformInfo(TEXT("IOS"),						TEXT("IOS"),				LOCTEXT("IOS", "iOS"),										EPlatformType::Game,		EPlatformFlags::None,			FPlatformIconPaths(TEXT("Launcher/iOS/Platform_iOS_24x"), TEXT("Launcher/iOS/Platform_iOS_128x")),										TEXT("")),
+	BuildPlatformInfo(TEXT("IOS"),						TEXT("IOS"),				LOCTEXT("IOS", "iOS"),										EPlatformType::Game,		EPlatformFlags::None,			FPlatformIconPaths(TEXT("Launcher/iOSTarget/Platform_iOS_24x"), TEXT("Launcher/iOSTarget/Platform_iOS_128x")),										TEXT("")),
 
-	BuildPlatformInfo(TEXT("PS4"),						TEXT("PS4"),				LOCTEXT("PS4", "PlayStation 4"),							EPlatformType::Game,		EPlatformFlags::None,			FPlatformIconPaths(TEXT("Launcher/PS4/Platform_PS4_24x"), TEXT("Launcher/PS4/Platform_PS4_128x")),										TEXT("")),
+	BuildPlatformInfo(TEXT("PS4"),						TEXT("PS4"),				LOCTEXT("PS4", "PlayStation 4"),							EPlatformType::Game,		EPlatformFlags::None,			FPlatformIconPaths(TEXT("Launcher/PS4Target/Platform_PS4_24x"), TEXT("Launcher/PS4Target/Platform_PS4_128x")),										TEXT("")),
 
-	BuildPlatformInfo(TEXT("XboxOne"),					TEXT("XboxOne"),			LOCTEXT("XboxOne", "Xbox One"),								EPlatformType::Game,		EPlatformFlags::None,			FPlatformIconPaths(TEXT("Launcher/XboxOne/Platform_XboxOne_24x"), TEXT("Launcher/XboxOne/Platform_XboxOne_128x")),						TEXT("")),
+	BuildPlatformInfo(TEXT("XboxOne"),					TEXT("XboxOne"),			LOCTEXT("XboxOne", "Xbox One"),								EPlatformType::Game,		EPlatformFlags::None,			FPlatformIconPaths(TEXT("Launcher/XboxOneTarget/Platform_XboxOne_24x"), TEXT("Launcher/XboxOneTarget/Platform_XboxOne_128x")),						TEXT("")),
 
-	BuildPlatformInfo(TEXT("HTML5"),					TEXT("HTML5"),				LOCTEXT("HTML5", "HTML5"),									EPlatformType::Game,		EPlatformFlags::None,			FPlatformIconPaths(TEXT("Launcher/HTML5/Platform_HTML5_24x"), TEXT("Launcher/HTML5/Platform_HTML5_128x")),								TEXT("")),
+	BuildPlatformInfo(TEXT("HTML5"),					TEXT("HTML5"),				LOCTEXT("HTML5", "HTML5"),									EPlatformType::Game,		EPlatformFlags::None,			FPlatformIconPaths(TEXT("Launcher/HTML5Target/Platform_HTML5_24x"), TEXT("Launcher/HTML5Target/Platform_HTML5_128x")),								TEXT("")),
 
-	//BuildPlatformInfo(TEXT("WinRT"),					TEXT("WinRT"),				LOCTEXT("WinRT", "Windows RT"),								EPlatformType::Game,		EPlatformFlags::None,			FPlatformIconPaths(TEXT("Launcher/Windows/Platform_WindowsNoEditor_24x"), TEXT("Launcher/Windows/Platform_WindowsNoEditor_128x")),		TEXT("")),
+	//BuildPlatformInfo(TEXT("WinRT"),					TEXT("WinRT"),				LOCTEXT("WinRT", "Windows RT"),								EPlatformType::Game,		EPlatformFlags::None,			FPlatformIconPaths(TEXT("Launcher/WindowsTarget/Platform_WindowsNoEditor_24x"), TEXT("Launcher/WindowsTarget/Platform_WindowsNoEditor_128x")),		TEXT("")),
 };
 
 } // anonymous namespace
