Small non-unity fix for OS X.

diff --git a/Engine/Source/Developer/DesktopPlatform/Private/Mac/MacNativeFeedbackContext.cpp b/Engine/Source/Developer/DesktopPlatform/Private/Mac/MacNativeFeedbackContext.cpp
--- a/Engine/Source/Developer/DesktopPlatform/Private/Mac/MacNativeFeedbackContext.cpp
+++ b/Engine/Source/Developer/DesktopPlatform/Private/Mac/MacNativeFeedbackContext.cpp
@@ -2,6 +2,7 @@
 
 #include "DesktopPlatformPrivatePCH.h"
 #include "MacNativeFeedbackContext.h"
+#include "MacApplication.h"
 
 @implementation FMacNativeFeedbackContextWindowController
 
