Remove the 'editor' suffix when recompiling a game project; we already specify -editorrecompile.

diff --git a/Engine/Source/Developer/DesktopPlatform/Private/DesktopPlatformBase.cpp b/Engine/Source/Developer/DesktopPlatform/Private/DesktopPlatformBase.cpp
--- a/Engine/Source/Developer/DesktopPlatform/Private/DesktopPlatformBase.cpp
+++ b/Engine/Source/Developer/DesktopPlatform/Private/DesktopPlatformBase.cpp
@@ -349,7 +349,7 @@ bool FDesktopPlatformBase::CleanGameProject(const FString &ProjectDir, FFeedback
 bool FDesktopPlatformBase::CompileGameProject(const FString& RootDir, const FString& ProjectFileName, FFeedbackContext* Warn)
 {
 	// Get the target name
-	FString Arguments = FString::Printf(TEXT("%sEditor %s %s"), *FPaths::GetBaseFilename(ProjectFileName), FModuleManager::Get().GetUBTConfiguration(), FPlatformMisc::GetUBTPlatform());
+	FString Arguments = FString::Printf(TEXT("%s %s %s"), *FPaths::GetBaseFilename(ProjectFileName), FModuleManager::Get().GetUBTConfiguration(), FPlatformMisc::GetUBTPlatform());
 
 	// Append the project name if it's a foreign project
 	if ( !ProjectFileName.IsEmpty() )
