Avoid an AutoSDK dependency on ITargetPlatformBase

diff --git a/Engine/Source/Developer/DesktopPlatform/Private/PlatformInfo.cpp b/Engine/Source/Developer/DesktopPlatform/Private/PlatformInfo.cpp
--- a/Engine/Source/Developer/DesktopPlatform/Private/PlatformInfo.cpp
+++ b/Engine/Source/Developer/DesktopPlatform/Private/PlatformInfo.cpp
@@ -11,7 +11,7 @@ namespace PlatformInfo
 namespace
 {
 
-FPlatformInfo BuildPlatformInfo(const FName& InPlatformInfoName, const FName& InTargetPlatformName, const FText& InDisplayName, const EPlatformType InPlatformType, const EPlatformFlags::Flags InPlatformFlags, const FPlatformIconPaths& InIconPaths, const FString& InUATCommandLine)
+FPlatformInfo BuildPlatformInfo(const FName& InPlatformInfoName, const FName& InTargetPlatformName, const FText& InDisplayName, const EPlatformType InPlatformType, const EPlatformFlags::Flags InPlatformFlags, const FPlatformIconPaths& InIconPaths, const FString& InUATCommandLine, const FString& InAutoSDKPath)
 {
 	FPlatformInfo PlatformInfo;
 
@@ -38,6 +38,7 @@ FPlatformInfo BuildPlatformInfo(const FName& InPlatformInfoName, const FName& In
 	PlatformInfo.PlatformFlags = InPlatformFlags;
 	PlatformInfo.IconPaths = InIconPaths;
 	PlatformInfo.UATCommandLine = InUATCommandLine;
+	PlatformInfo.AutoSDKPath = InAutoSDKPath;
 
 	// Generate the icon style names for FEditorStyle
 	PlatformInfo.IconPaths.NormalStyleName = *FString::Printf(TEXT("Launcher.Platform_%s"), *InPlatformInfoNameString);
@@ -48,39 +49,39 @@ FPlatformInfo BuildPlatformInfo(const FName& InPlatformInfoName, const FName& In
 }
 
 static const FPlatformInfo PlatformInfoArray[] = {
-	//				  PlatformInfoName					TargetPlatformName			DisplayName													PlatformType				PlatformFlags					IconPaths																																			UATCommandLine
-	BuildPlatformInfo(TEXT("WindowsNoEditor"),			TEXT("WindowsNoEditor"),	LOCTEXT("WindowsNoEditor", "Windows"),						EPlatformType::Game,		EPlatformFlags::None,			FPlatformIconPaths(TEXT("Launcher/WindowsTarget/Platform_WindowsNoEditor_24x"), TEXT("Launcher/WindowsTarget/Platform_WindowsNoEditor_128x")),		TEXT("")),
-	BuildPlatformInfo(TEXT("WindowsNoEditor_Win32"),	TEXT("WindowsNoEditor"),	LOCTEXT("WindowsNoEditor_Win32", "Windows (32-bit)"),		EPlatformType::Game,		EPlatformFlags::BuildFlavor,	FPlatformIconPaths(TEXT("Launcher/WindowsTarget/Platform_WindowsNoEditor_24x"), TEXT("Launcher/WindowsTarget/Platform_WindowsNoEditor_128x")),		TEXT("-targetplatform=Win32")),
-	BuildPlatformInfo(TEXT("WindowsNoEditor_Win64"),	TEXT("WindowsNoEditor"),	LOCTEXT("WindowsNoEditor_Win64", "Windows (64-bit)"),		EPlatformType::Game,		EPlatformFlags::BuildFlavor,	FPlatformIconPaths(TEXT("Launcher/WindowsTarget/Platform_WindowsNoEditor_24x"), TEXT("Launcher/WindowsTarget/Platform_WindowsNoEditor_128x")),		TEXT("-targetplatform=Win64")),
-	BuildPlatformInfo(TEXT("Windows"),					TEXT("Windows"),			LOCTEXT("WindowsEditor", "Windows (Editor)"),				EPlatformType::Editor,		EPlatformFlags::None,			FPlatformIconPaths(TEXT("Launcher/WindowsTarget/Platform_Windows_24x"), TEXT("Launcher/WindowsTarget/Platform_Windows_128x")),						TEXT("")),
-	BuildPlatformInfo(TEXT("WindowsClient"),			TEXT("WindowsClient"),		LOCTEXT("WindowsClient", "Windows (Client-only)"),			EPlatformType::Client,		EPlatformFlags::None,			FPlatformIconPaths(TEXT("Launcher/WindowsTarget/Platform_Windows_24x"), TEXT("Launcher/WindowsTarget/Platform_Windows_128x")),						TEXT("")),
-	BuildPlatformInfo(TEXT("WindowsServer"),			TEXT("WindowsServer"),		LOCTEXT("WindowsServer", "Windows (Dedicated Server)"),		EPlatformType::Server,		EPlatformFlags::None,			FPlatformIconPaths(TEXT("Launcher/WindowsTarget/Platform_WindowsServer_24x"), TEXT("Launcher/WindowsTarget/Platform_WindowsServer_128x")),			TEXT("")),
-
-	BuildPlatformInfo(TEXT("MacNoEditor"),				TEXT("MacNoEditor"),		LOCTEXT("MacNoEditor", "Mac"),								EPlatformType::Game,		EPlatformFlags::None,			FPlatformIconPaths(TEXT("Launcher/MacTarget/Platform_Mac_24x"), TEXT("Launcher/MacTarget/Platform_Mac_128x")),										TEXT("")),
-	BuildPlatformInfo(TEXT("Mac"),						TEXT("Mac"),				LOCTEXT("MacEditor", "Mac (Editor)"),						EPlatformType::Editor,		EPlatformFlags::None,			FPlatformIconPaths(TEXT("Launcher/MacTarget/Platform_Mac_24x"), TEXT("Launcher/MacTarget/Platform_Mac_128x")),										TEXT("")),
-	BuildPlatformInfo(TEXT("MacClient"),				TEXT("MacClient"),			LOCTEXT("MacClient", "Mac (Client-only)"),					EPlatformType::Client,		EPlatformFlags::None,			FPlatformIconPaths(TEXT("Launcher/MacTarget/Platform_Mac_24x"), TEXT("Launcher/MacTarget/Platform_Mac_128x")),										TEXT("")),
-	BuildPlatformInfo(TEXT("MacServer"),				TEXT("MacServer"),			LOCTEXT("MacServer", "Mac (Dedicated Server)"),				EPlatformType::Server,		EPlatformFlags::None,			FPlatformIconPaths(TEXT("Launcher/MacTarget/Platform_Mac_24x"), TEXT("Launcher/MacTarget/Platform_Mac_128x")),										TEXT("")),
-
-	BuildPlatformInfo(TEXT("LinuxNoEditor"),			TEXT("LinuxNoEditor"),		LOCTEXT("LinuxNoEditor", "Linux"),							EPlatformType::Game,		EPlatformFlags::None,			FPlatformIconPaths(TEXT("Launcher/LinuxTarget/Platform_Linux_24x"), TEXT("Launcher/LinuxTarget/Platform_Linux_128x")),								TEXT("")),
-	BuildPlatformInfo(TEXT("Linux"),					TEXT("Linux"),				LOCTEXT("LinuxEditor", "Linux (Editor)"),					EPlatformType::Editor,		EPlatformFlags::None,			FPlatformIconPaths(TEXT("Launcher/LinuxTarget/Platform_Linux_24x"), TEXT("Launcher/LinuxTarget/Platform_Linux_128x")),								TEXT("")),
-	BuildPlatformInfo(TEXT("LinuxClient"),				TEXT("LinuxClient"),		LOCTEXT("LinuxClient", "Linux (Client-only)"),				EPlatformType::Client,		EPlatformFlags::None,			FPlatformIconPaths(TEXT("Launcher/LinuxTarget/Platform_Linux_24x"), TEXT("Launcher/LinuxTarget/Platform_Linux_128x")),								TEXT("")),
-	BuildPlatformInfo(TEXT("LinuxServer"),				TEXT("LinuxServer"),		LOCTEXT("LinuxServer", "Linux (Dedicated Server)"),			EPlatformType::Server,		EPlatformFlags::None,			FPlatformIconPaths(TEXT("Launcher/LinuxTarget/Platform_Linux_24x"), TEXT("Launcher/LinuxTarget/Platform_Linux_128x")),								TEXT("")),
+	//				  PlatformInfoName					TargetPlatformName			DisplayName													PlatformType				PlatformFlags					IconPaths																																			UATCommandLine											AutoSDKPath
+	BuildPlatformInfo(TEXT("WindowsNoEditor"),			TEXT("WindowsNoEditor"),	LOCTEXT("WindowsNoEditor", "Windows"),						EPlatformType::Game,		EPlatformFlags::None,			FPlatformIconPaths(TEXT("Launcher/WindowsTarget/Platform_WindowsNoEditor_24x"), TEXT("Launcher/WindowsTarget/Platform_WindowsNoEditor_128x")),		TEXT(""),												TEXT("")),
+	BuildPlatformInfo(TEXT("WindowsNoEditor_Win32"),	TEXT("WindowsNoEditor"),	LOCTEXT("WindowsNoEditor_Win32", "Windows (32-bit)"),		EPlatformType::Game,		EPlatformFlags::BuildFlavor,	FPlatformIconPaths(TEXT("Launcher/WindowsTarget/Platform_WindowsNoEditor_24x"), TEXT("Launcher/WindowsTarget/Platform_WindowsNoEditor_128x")),		TEXT("-targetplatform=Win32"),							TEXT("")),
+	BuildPlatformInfo(TEXT("WindowsNoEditor_Win64"),	TEXT("WindowsNoEditor"),	LOCTEXT("WindowsNoEditor_Win64", "Windows (64-bit)"),		EPlatformType::Game,		EPlatformFlags::BuildFlavor,	FPlatformIconPaths(TEXT("Launcher/WindowsTarget/Platform_WindowsNoEditor_24x"), TEXT("Launcher/WindowsTarget/Platform_WindowsNoEditor_128x")),		TEXT("-targetplatform=Win64"),							TEXT("")),
+	BuildPlatformInfo(TEXT("Windows"),					TEXT("Windows"),			LOCTEXT("WindowsEditor", "Windows (Editor)"),				EPlatformType::Editor,		EPlatformFlags::None,			FPlatformIconPaths(TEXT("Launcher/WindowsTarget/Platform_Windows_24x"), TEXT("Launcher/WindowsTarget/Platform_Windows_128x")),						TEXT(""),												TEXT("")),
+	BuildPlatformInfo(TEXT("WindowsClient"),			TEXT("WindowsClient"),		LOCTEXT("WindowsClient", "Windows (Client-only)"),			EPlatformType::Client,		EPlatformFlags::None,			FPlatformIconPaths(TEXT("Launcher/WindowsTarget/Platform_Windows_24x"), TEXT("Launcher/WindowsTarget/Platform_Windows_128x")),						TEXT(""),												TEXT("")),
+	BuildPlatformInfo(TEXT("WindowsServer"),			TEXT("WindowsServer"),		LOCTEXT("WindowsServer", "Windows (Dedicated Server)"),		EPlatformType::Server,		EPlatformFlags::None,			FPlatformIconPaths(TEXT("Launcher/WindowsTarget/Platform_WindowsServer_24x"), TEXT("Launcher/WindowsTarget/Platform_WindowsServer_128x")),			TEXT(""),												TEXT("")),
+
+	BuildPlatformInfo(TEXT("MacNoEditor"),				TEXT("MacNoEditor"),		LOCTEXT("MacNoEditor", "Mac"),								EPlatformType::Game,		EPlatformFlags::None,			FPlatformIconPaths(TEXT("Launcher/MacTarget/Platform_Mac_24x"), TEXT("Launcher/MacTarget/Platform_Mac_128x")),										TEXT(""),												TEXT("")),
+	BuildPlatformInfo(TEXT("Mac"),						TEXT("Mac"),				LOCTEXT("MacEditor", "Mac (Editor)"),						EPlatformType::Editor,		EPlatformFlags::None,			FPlatformIconPaths(TEXT("Launcher/MacTarget/Platform_Mac_24x"), TEXT("Launcher/MacTarget/Platform_Mac_128x")),										TEXT(""),												TEXT("")),
+	BuildPlatformInfo(TEXT("MacClient"),				TEXT("MacClient"),			LOCTEXT("MacClient", "Mac (Client-only)"),					EPlatformType::Client,		EPlatformFlags::None,			FPlatformIconPaths(TEXT("Launcher/MacTarget/Platform_Mac_24x"), TEXT("Launcher/MacTarget/Platform_Mac_128x")),										TEXT(""),												TEXT("")),
+	BuildPlatformInfo(TEXT("MacServer"),				TEXT("MacServer"),			LOCTEXT("MacServer", "Mac (Dedicated Server)"),				EPlatformType::Server,		EPlatformFlags::None,			FPlatformIconPaths(TEXT("Launcher/MacTarget/Platform_Mac_24x"), TEXT("Launcher/MacTarget/Platform_Mac_128x")),										TEXT(""),												TEXT("")),
+
+	BuildPlatformInfo(TEXT("LinuxNoEditor"),			TEXT("LinuxNoEditor"),		LOCTEXT("LinuxNoEditor", "Linux"),							EPlatformType::Game,		EPlatformFlags::None,			FPlatformIconPaths(TEXT("Launcher/LinuxTarget/Platform_Linux_24x"), TEXT("Launcher/LinuxTarget/Platform_Linux_128x")),								TEXT(""),												TEXT("Linux_x64")),
+	BuildPlatformInfo(TEXT("Linux"),					TEXT("Linux"),				LOCTEXT("LinuxEditor", "Linux (Editor)"),					EPlatformType::Editor,		EPlatformFlags::None,			FPlatformIconPaths(TEXT("Launcher/LinuxTarget/Platform_Linux_24x"), TEXT("Launcher/LinuxTarget/Platform_Linux_128x")),								TEXT(""),												TEXT("Linux_x64")),
+	BuildPlatformInfo(TEXT("LinuxClient"),				TEXT("LinuxClient"),		LOCTEXT("LinuxClient", "Linux (Client-only)"),				EPlatformType::Client,		EPlatformFlags::None,			FPlatformIconPaths(TEXT("Launcher/LinuxTarget/Platform_Linux_24x"), TEXT("Launcher/LinuxTarget/Platform_Linux_128x")),								TEXT(""),												TEXT("Linux_x64")),
+	BuildPlatformInfo(TEXT("LinuxServer"),				TEXT("LinuxServer"),		LOCTEXT("LinuxServer", "Linux (Dedicated Server)"),			EPlatformType::Server,		EPlatformFlags::None,			FPlatformIconPaths(TEXT("Launcher/LinuxTarget/Platform_Linux_24x"), TEXT("Launcher/LinuxTarget/Platform_Linux_128x")),								TEXT(""),												TEXT("Linux_x64")),
 	
-	BuildPlatformInfo(TEXT("Android"),					TEXT("Android"),			LOCTEXT("Android", "Android"),								EPlatformType::Game,		EPlatformFlags::None,			FPlatformIconPaths(TEXT("Launcher/AndroidTarget/Platform_Android_24x"), TEXT("Launcher/AndroidTarget/Platform_Android_128x")),						TEXT("")),
-	BuildPlatformInfo(TEXT("Android_All"),				TEXT("Android"),			LOCTEXT("Android_All", "Android (All)"),					EPlatformType::Game,		EPlatformFlags::CookFlavor,		FPlatformIconPaths(TEXT("Launcher/AndroidTarget/Platform_Android_24x"), TEXT("Launcher/AndroidTarget/Platform_Android_128x")),						TEXT("")),
-	BuildPlatformInfo(TEXT("Android_ATC"),				TEXT("Android_ATC"),		LOCTEXT("Android_ATC", "Android (ATC)"),					EPlatformType::Game,		EPlatformFlags::CookFlavor,		FPlatformIconPaths(TEXT("Launcher/AndroidTarget/Platform_Android_ATC_24x"), TEXT("Launcher/AndroidTarget/Platform_Android_128x")),					TEXT("-targetplatform=Android -cookflavor=ATC")),
-	BuildPlatformInfo(TEXT("Android_DXT"),				TEXT("Android_DXT"),		LOCTEXT("Android_DXT", "Android (DXT)"),					EPlatformType::Game,		EPlatformFlags::CookFlavor,		FPlatformIconPaths(TEXT("Launcher/AndroidTarget/Platform_Android_DXT_24x"), TEXT("Launcher/AndroidTarget/Platform_Android_128x")),					TEXT("-targetplatform=Android -cookflavor=DXT")),
-	BuildPlatformInfo(TEXT("Android_ETC1"),				TEXT("Android_ETC1"),		LOCTEXT("Android_ETC1", "Android (ETC1)"),					EPlatformType::Game,		EPlatformFlags::CookFlavor,		FPlatformIconPaths(TEXT("Launcher/AndroidTarget/Platform_Android_ETC1_24x"), TEXT("Launcher/AndroidTarget/Platform_Android_128x")),					TEXT("-targetplatform=Android -cookflavor=ETC1")),
-	BuildPlatformInfo(TEXT("Android_ETC2"),				TEXT("Android_ETC2"),		LOCTEXT("Android_ETC2", "Android (ETC2)"),					EPlatformType::Game,		EPlatformFlags::CookFlavor,		FPlatformIconPaths(TEXT("Launcher/AndroidTarget/Platform_Android_ETC2_24x"), TEXT("Launcher/AndroidTarget/Platform_Android_128x")),					TEXT("-targetplatform=Android -cookflavor=ETC2")),
-	BuildPlatformInfo(TEXT("Android_PVRTC"),			TEXT("Android_PVRTC"),		LOCTEXT("Android_PVRTC", "Android (PVRTC)"),				EPlatformType::Game,		EPlatformFlags::CookFlavor,		FPlatformIconPaths(TEXT("Launcher/AndroidTarget/Platform_Android_PVRTC_24x"), TEXT("Launcher/AndroidTarget/Platform_Android_128x")),				TEXT("-targetplatform=Android -cookflavor=PVRTC")),
+	BuildPlatformInfo(TEXT("Android"),					TEXT("Android"),			LOCTEXT("Android", "Android"),								EPlatformType::Game,		EPlatformFlags::None,			FPlatformIconPaths(TEXT("Launcher/AndroidTarget/Platform_Android_24x"), TEXT("Launcher/AndroidTarget/Platform_Android_128x")),						TEXT(""),												TEXT("Android")),
+	BuildPlatformInfo(TEXT("Android_All"),				TEXT("Android"),			LOCTEXT("Android_All", "Android (All)"),					EPlatformType::Game,		EPlatformFlags::CookFlavor,		FPlatformIconPaths(TEXT("Launcher/AndroidTarget/Platform_Android_24x"), TEXT("Launcher/AndroidTarget/Platform_Android_128x")),						TEXT(""),												TEXT("Android")),
+	BuildPlatformInfo(TEXT("Android_ATC"),				TEXT("Android_ATC"),		LOCTEXT("Android_ATC", "Android (ATC)"),					EPlatformType::Game,		EPlatformFlags::CookFlavor,		FPlatformIconPaths(TEXT("Launcher/AndroidTarget/Platform_Android_ATC_24x"), TEXT("Launcher/AndroidTarget/Platform_Android_128x")),					TEXT("-targetplatform=Android -cookflavor=ATC"),		TEXT("Android")),
+	BuildPlatformInfo(TEXT("Android_DXT"),				TEXT("Android_DXT"),		LOCTEXT("Android_DXT", "Android (DXT)"),					EPlatformType::Game,		EPlatformFlags::CookFlavor,		FPlatformIconPaths(TEXT("Launcher/AndroidTarget/Platform_Android_DXT_24x"), TEXT("Launcher/AndroidTarget/Platform_Android_128x")),					TEXT("-targetplatform=Android -cookflavor=DXT"),		TEXT("Android")),
+	BuildPlatformInfo(TEXT("Android_ETC1"),				TEXT("Android_ETC1"),		LOCTEXT("Android_ETC1", "Android (ETC1)"),					EPlatformType::Game,		EPlatformFlags::CookFlavor,		FPlatformIconPaths(TEXT("Launcher/AndroidTarget/Platform_Android_ETC1_24x"), TEXT("Launcher/AndroidTarget/Platform_Android_128x")),					TEXT("-targetplatform=Android -cookflavor=ETC1"),		TEXT("Android")),
+	BuildPlatformInfo(TEXT("Android_ETC2"),				TEXT("Android_ETC2"),		LOCTEXT("Android_ETC2", "Android (ETC2)"),					EPlatformType::Game,		EPlatformFlags::CookFlavor,		FPlatformIconPaths(TEXT("Launcher/AndroidTarget/Platform_Android_ETC2_24x"), TEXT("Launcher/AndroidTarget/Platform_Android_128x")),					TEXT("-targetplatform=Android -cookflavor=ETC2"),		TEXT("Android")),
+	BuildPlatformInfo(TEXT("Android_PVRTC"),			TEXT("Android_PVRTC"),		LOCTEXT("Android_PVRTC", "Android (PVRTC)"),				EPlatformType::Game,		EPlatformFlags::CookFlavor,		FPlatformIconPaths(TEXT("Launcher/AndroidTarget/Platform_Android_PVRTC_24x"), TEXT("Launcher/AndroidTarget/Platform_Android_128x")),				TEXT("-targetplatform=Android -cookflavor=PVRTC"),		TEXT("Android")),
 
-	BuildPlatformInfo(TEXT("IOS"),						TEXT("IOS"),				LOCTEXT("IOS", "iOS"),										EPlatformType::Game,		EPlatformFlags::None,			FPlatformIconPaths(TEXT("Launcher/iOSTarget/Platform_iOS_24x"), TEXT("Launcher/iOSTarget/Platform_iOS_128x")),										TEXT("")),
+	BuildPlatformInfo(TEXT("IOS"),						TEXT("IOS"),				LOCTEXT("IOS", "iOS"),										EPlatformType::Game,		EPlatformFlags::None,			FPlatformIconPaths(TEXT("Launcher/iOSTarget/Platform_iOS_24x"), TEXT("Launcher/iOSTarget/Platform_iOS_128x")),										TEXT(""),												TEXT("")),
 
-	BuildPlatformInfo(TEXT("PS4"),						TEXT("PS4"),				LOCTEXT("PS4", "PlayStation 4"),							EPlatformType::Game,		EPlatformFlags::None,			FPlatformIconPaths(TEXT("Launcher/PS4Target/Platform_PS4_24x"), TEXT("Launcher/PS4Target/Platform_PS4_128x")),										TEXT("")),
+	BuildPlatformInfo(TEXT("PS4"),						TEXT("PS4"),				LOCTEXT("PS4", "PlayStation 4"),							EPlatformType::Game,		EPlatformFlags::None,			FPlatformIconPaths(TEXT("Launcher/PS4Target/Platform_PS4_24x"), TEXT("Launcher/PS4Target/Platform_PS4_128x")),										TEXT(""),												TEXT("PS4")),
 
-	BuildPlatformInfo(TEXT("XboxOne"),					TEXT("XboxOne"),			LOCTEXT("XboxOne", "Xbox One"),								EPlatformType::Game,		EPlatformFlags::None,			FPlatformIconPaths(TEXT("Launcher/XboxOneTarget/Platform_XboxOne_24x"), TEXT("Launcher/XboxOneTarget/Platform_XboxOne_128x")),						TEXT("")),
+	BuildPlatformInfo(TEXT("XboxOne"),					TEXT("XboxOne"),			LOCTEXT("XboxOne", "Xbox One"),								EPlatformType::Game,		EPlatformFlags::None,			FPlatformIconPaths(TEXT("Launcher/XboxOneTarget/Platform_XboxOne_24x"), TEXT("Launcher/XboxOneTarget/Platform_XboxOne_128x")),						TEXT(""),												TEXT("XboxOne")),
 
-	BuildPlatformInfo(TEXT("HTML5"),					TEXT("HTML5"),				LOCTEXT("HTML5", "HTML5"),									EPlatformType::Game,		EPlatformFlags::None,			FPlatformIconPaths(TEXT("Launcher/HTML5Target/Platform_HTML5_24x"), TEXT("Launcher/HTML5Target/Platform_HTML5_128x")),								TEXT("")),
+	BuildPlatformInfo(TEXT("HTML5"),					TEXT("HTML5"),				LOCTEXT("HTML5", "HTML5"),									EPlatformType::Game,		EPlatformFlags::None,			FPlatformIconPaths(TEXT("Launcher/HTML5Target/Platform_HTML5_24x"), TEXT("Launcher/HTML5Target/Platform_HTML5_128x")),								TEXT(""),												TEXT("HTML5")),
 
 	//BuildPlatformInfo(TEXT("WinRT"),					TEXT("WinRT"),				LOCTEXT("WinRT", "Windows RT"),								EPlatformType::Game,		EPlatformFlags::None,			FPlatformIconPaths(TEXT("Launcher/WindowsTarget/Platform_WindowsNoEditor_24x"), TEXT("Launcher/WindowsTarget/Platform_WindowsNoEditor_128x")),		TEXT("")),
 };
diff --git a/Engine/Source/Developer/DesktopPlatform/Public/PlatformInfo.h b/Engine/Source/Developer/DesktopPlatform/Public/PlatformInfo.h
--- a/Engine/Source/Developer/DesktopPlatform/Public/PlatformInfo.h
+++ b/Engine/Source/Developer/DesktopPlatform/Public/PlatformInfo.h
@@ -131,6 +131,9 @@ namespace PlatformInfo
 		/** Additional argument string data to append to UAT commands relating to this platform */
 		FString UATCommandLine;
 
+		/** Path under CarefullyRedist for the SDK.  FString so case sensitive platforms don't get messed up by a pre-existing FName of a different casing. */
+		FString AutoSDKPath;
+
 		/** Returns true if this platform is vanilla */
 		FORCEINLINE bool IsVanilla() const
 		{
diff --git a/Engine/Source/Developer/TargetPlatform/Private/TargetPlatformManagerModule.cpp b/Engine/Source/Developer/TargetPlatform/Private/TargetPlatformManagerModule.cpp
--- a/Engine/Source/Developer/TargetPlatform/Private/TargetPlatformManagerModule.cpp
+++ b/Engine/Source/Developer/TargetPlatform/Private/TargetPlatformManagerModule.cpp
@@ -5,6 +5,7 @@
 =============================================================================*/
 
 #include "TargetPlatformPrivatePCH.h"
+#include "PlatformInfo.h"
 
 DEFINE_LOG_CATEGORY_STATIC(LogTargetPlatformManager, Log, All);
 
@@ -38,7 +39,20 @@ public:
 		if (!FUBTInvoker::InvokeUnrealBuildToolSync(UBTParams, *GLog, true, UBTReturnCode, UBTOutput))
 		{
 			UE_LOG(LogTargetPlatformManager, Fatal, TEXT("Failed to run UBT to check SDK status!"));
-		}		
+		}
+
+		// we have to setup our local environment according to AutoSDKs or the ITargetPlatform's IsSDkInstalled calls may fail
+		// before we get a change to setup for a given platform.  Use the platforminfo list to avoid any kind of interdependency.
+		int32 NumPlatforms;
+		const PlatformInfo::FPlatformInfo* PlatformInfoArray = PlatformInfo::GetPlatformInfoArray(NumPlatforms);
+		for (int32 i = 0; i < NumPlatforms; ++i)
+		{
+			const PlatformInfo::FPlatformInfo& PlatformInfo = PlatformInfoArray[i];
+			if (PlatformInfo.AutoSDKPath.Len() > 0)
+			{
+				SetupAndValidateAutoSDK(PlatformInfo.AutoSDKPath);
+			}
+		}
 #endif
 
 		GetTargetPlatforms();
@@ -528,8 +542,8 @@ protected:
 				if (Platform != NULL)
 				{
 					// would like to move this check to GetActiveTargetPlatforms, but too many things cache this result
-					// this setup will become faster after TTP 341897 is complete.
-					if (SetupAndValidateAutoSDK(*Platform))
+					// this setup will become faster after TTP 341897 is complete.					
+					if (SetupAndValidateAutoSDK(Platform->GetPlatformInfo().AutoSDKPath))
 					{
 						Platforms.Add(Platform);
 					}
@@ -538,24 +552,30 @@ protected:
 		}
 	}
 
-	bool SetupAndValidateAutoSDK(ITargetPlatform& Platform)
+	bool SetupAndValidateAutoSDK(const FString& AutoSDKPath)
 	{
 #if AUTOSDKS_ENABLED
 		bool bValidSDK = false;
-		// Invoke UBT to perform SDK switching, or detect that a proper manual SDK is already setup.
-		FString PlatformName = Platform.IniPlatformName();
-		FName PlatformFName(*PlatformName);
+		if (AutoSDKPath.Len() > 0)
+		{		
+			FName PlatformFName(*AutoSDKPath);
 
-		// cache result of the last setup attempt to avoid calling UBT all the time.
-		bool* bPreviousSetupSuccessful = PlatformsSetup.Find(PlatformFName);
-		if (bPreviousSetupSuccessful)
-		{
-			bValidSDK = *bPreviousSetupSuccessful;
+			// cache result of the last setup attempt to avoid calling UBT all the time.
+			bool* bPreviousSetupSuccessful = PlatformsSetup.Find(PlatformFName);
+			if (bPreviousSetupSuccessful)
+			{
+				bValidSDK = *bPreviousSetupSuccessful;
+			}
+			else
+			{
+				bValidSDK = SetupEnvironmentFromAutoSDK(AutoSDKPath);
+				PlatformsSetup.Add(PlatformFName, bValidSDK);
+			}
 		}
 		else
 		{
-			bValidSDK = SetupEnvironmentFromAutoSDK(Platform);
-			PlatformsSetup.Add(PlatformFName, bValidSDK);
+			// if a platform has no AutoSDKPath, then just assume the SDK is installed, we have no basis for determining it.
+			bValidSDK = true;
 		}
 		return bValidSDK;
 #else
@@ -563,7 +583,7 @@ protected:
 #endif // AUTOSDKS_ENABLED
 	}
 	
-	bool SetupEnvironmentFromAutoSDK(ITargetPlatform& Platform)
+	bool SetupEnvironmentFromAutoSDK(const FString& AutoSDKPath)
 	{						
 #if AUTOSDKS_ENABLED
 		static const FString SDKRootEnvFar(TEXT("UE_SDKS_ROOT"));		
@@ -575,31 +595,16 @@ protected:
 		if (SDKPath[0] == 0)
 		{
 			return true;
-		}
-
-		// if this platform hasn't been setup with AutoSDK handling, then we just assume it's setup
-		// properly.
-		if (!Platform.SupportsAutoSDK())
-		{
-			return true;
-		}
-
-		// Invoke UBT to perform SDK switching, or detect that a proper manual SDK is already setup.
-		FString PlatformName = Platform.IniPlatformName();
-		FName PlatformFName(*PlatformName);		
+		}		
 
+		// Invoke UBT to perform SDK switching, or detect that a proper manual SDK is already setup.				
 #if PLATFORM_WINDOWS
 		FString HostPlatform(TEXT("HostWin64"));
 #else
 #error Fill in your host platform directory
-#endif
-
-		if (PlatformName.Find(TEXT("Linux"), ESearchCase::IgnoreCase) != -1)
-		{
-			PlatformName = TEXT("Linux_x64");
-		}
+#endif		
 
-		FString TargetSDKRoot = FPaths::Combine(SDKPath, *HostPlatform, *PlatformName);
+		FString TargetSDKRoot = FPaths::Combine(SDKPath, *HostPlatform, *AutoSDKPath);
 		static const FString SDKInstallManifestFileName(TEXT("CurrentlyInstalled.txt"));
 		FString SDKInstallManifestFilePath = FPaths::Combine(*TargetSDKRoot, *SDKInstallManifestFileName);
 
@@ -624,20 +629,20 @@ protected:
 			
 			if (FileLines.Num() != 2)
 			{
-				UE_LOG(LogTargetPlatformManager, Error, TEXT("Malformed install manifest file for Platform %s"), *PlatformName);				
+				UE_LOG(LogTargetPlatformManager, Error, TEXT("Malformed install manifest file for Platform %s"), *AutoSDKPath);				
 				return false;
 			}
 
 			static const FString ManualSDKString(TEXT("ManualSDK"));
 			if (FileLines[1].Compare(ManualSDKString, ESearchCase::IgnoreCase) == 0)
 			{
-				UE_LOG(LogTargetPlatformManager, Verbose, TEXT("Platform %s has manual sdk install"), *PlatformName);				
+				UE_LOG(LogTargetPlatformManager, Verbose, TEXT("Platform %s has manual sdk install"), *AutoSDKPath);				
 				return true;
 			}
 		}
 		else
 		{	
-			UE_LOG(LogTargetPlatformManager, Error, TEXT("install manifest file for Platform %s not found."), *PlatformName);			
+			UE_LOG(LogTargetPlatformManager, Error, TEXT("install manifest file for Platform %s not found."), *AutoSDKPath);			
 			return false;			
 		}		
 
@@ -766,11 +771,11 @@ protected:
 		}
 		else
 		{
-			UE_LOG(LogTargetPlatformManager, Error, TEXT("OutputEnvVars.txt not found for platform: '%s'"), *PlatformName);			
+			UE_LOG(LogTargetPlatformManager, Error, TEXT("OutputEnvVars.txt not found for platform: '%s'"), *AutoSDKPath);			
 			return false;
 		}
 
-		UE_LOG(LogTargetPlatformManager, Verbose, TEXT("Platform %s has auto sdk install"), *PlatformName);		
+		UE_LOG(LogTargetPlatformManager, Verbose, TEXT("Platform %s has auto sdk install"), *AutoSDKPath);		
 		return true;
 #else
 		return true;
diff --git a/Engine/Source/Developer/TargetPlatform/TargetPlatform.Build.cs b/Engine/Source/Developer/TargetPlatform/TargetPlatform.Build.cs
--- a/Engine/Source/Developer/TargetPlatform/TargetPlatform.Build.cs
+++ b/Engine/Source/Developer/TargetPlatform/TargetPlatform.Build.cs
@@ -7,9 +7,8 @@ public class TargetPlatform : ModuleRules
 {
 	public TargetPlatform(TargetInfo Target)
 	{
-		PrivateDependencyModuleNames.Add("Core");
-
-		PublicIncludePathModuleNames.Add("DesktopPlatform");
+		PrivateDependencyModuleNames.Add("Core");		
+        PublicDependencyModuleNames.Add("DesktopPlatform");
 
 		if (!UEBuildConfiguration.bBuildRequiresCookedData)
 		{
