Support for Linux in DesktopPlatform module.

- Part of PR #306 by amigo and #ue4linux folks.
- LinuxNativeDialogs on hold until clarified with the author (their Build.cs needs to be modified to add WITH_LINUX_NATIVE_DIALOGS=1).

diff --git a/Engine/Source/Developer/DesktopPlatform/DesktopPlatform.Build.cs b/Engine/Source/Developer/DesktopPlatform/DesktopPlatform.Build.cs
--- a/Engine/Source/Developer/DesktopPlatform/DesktopPlatform.Build.cs
+++ b/Engine/Source/Developer/DesktopPlatform/DesktopPlatform.Build.cs
@@ -12,6 +12,7 @@ public class DesktopPlatform : ModuleRules
 
 		if (Target.Platform == UnrealTargetPlatform.Linux)
 		{
+            // on hold AddThirdPartyPrivateStaticDependencies(Target, "LinuxNativeDialogs");
             AddThirdPartyPrivateStaticDependencies(Target, "SDL2");
 		}
 	}
diff --git a/Engine/Source/Developer/DesktopPlatform/Private/DesktopPlatformBase.cpp b/Engine/Source/Developer/DesktopPlatform/Private/DesktopPlatformBase.cpp
--- a/Engine/Source/Developer/DesktopPlatform/Private/DesktopPlatformBase.cpp
+++ b/Engine/Source/Developer/DesktopPlatform/Private/DesktopPlatformBase.cpp
@@ -380,6 +380,8 @@ bool FDesktopPlatformBase::GenerateProjectFiles(const FString& RootDir, const FS
 {
 #if PLATFORM_MAC
 	FString Arguments = TEXT("-xcodeprojectfile");
+#elif PLATFORM_LINUX
+	FString Arguments = TEXT("-makefile");
 #else
 	FString Arguments = TEXT("-projectfiles");
 #endif
diff --git a/Engine/Source/Developer/DesktopPlatform/Private/Linux/DesktopPlatformLinux.cpp b/Engine/Source/Developer/DesktopPlatform/Private/Linux/DesktopPlatformLinux.cpp
--- a/Engine/Source/Developer/DesktopPlatform/Private/Linux/DesktopPlatformLinux.cpp
+++ b/Engine/Source/Developer/DesktopPlatform/Private/Linux/DesktopPlatformLinux.cpp
@@ -1,31 +1,84 @@
 // Copyright 1998-2014 Epic Games, Inc. All Rights Reserved.
 
 #include "DesktopPlatformPrivatePCH.h"
+#include "LinuxApplication.h"
+#include "FeedbackContextMarkup.h"
+//#include "LinuxNativeFeedbackContext.h"
+// custom dialogs
+#if WITH_LINUX_NATIVE_DIALOGS
+#include "UNativeDialogs.h"
+#endif // WITH_LINUX_NATIVE_DIALOGS
+#include "SDL.h"
 
 #define LOCTEXT_NAMESPACE "DesktopPlatform"
+#define MAX_FILETYPES_STR 4096
+#define MAX_FILENAME_STR 65536
 
-bool FDesktopPlatformLinux::OpenFileDialog(const void* ParentWindowHandle, const FString& DialogTitle, const FString& DefaultPath, const FString& DefaultFile, const FString& FileTypes, uint32 Flags, TArray<FString>& OutFilenames)
+bool FDesktopPlatformLinux::OpenFileDialog(const void* ParentWindowHandle, const FString& DialogTitle, const FString& DefaultPath, const FString& DefaultFile, const FString& FileTypes, uint32 Flags, TArray<FString>& OutFilenames, int32& OutFilterIndex)
 {
-	unimplemented();
-	return false;
+	return FileDialogShared(false, ParentWindowHandle, DialogTitle, DefaultPath, DefaultFile, FileTypes, Flags, OutFilenames, OutFilterIndex);
 }
 
-bool FDesktopPlatformLinux::OpenFileDialog(const void* ParentWindowHandle, const FString& DialogTitle, const FString& DefaultPath, const FString& DefaultFile, const FString& FileTypes, uint32 Flags, TArray<FString>& OutFilenames, int32& OutFilterIndex)
+bool FDesktopPlatformLinux::OpenFileDialog(const void* ParentWindowHandle, const FString& DialogTitle, const FString& DefaultPath, const FString& DefaultFile, const FString& FileTypes, uint32 Flags, TArray<FString>& OutFilenames)
 {
-	unimplemented();
-	return false;
+	int32 DummyFilterIndex;
+	return FileDialogShared(false, ParentWindowHandle, DialogTitle, DefaultPath, DefaultFile, FileTypes, Flags, OutFilenames, DummyFilterIndex);
 }
 
 bool FDesktopPlatformLinux::SaveFileDialog(const void* ParentWindowHandle, const FString& DialogTitle, const FString& DefaultPath, const FString& DefaultFile, const FString& FileTypes, uint32 Flags, TArray<FString>& OutFilenames)
 {
-	unimplemented();
-	return false;
+	int32 DummyFilterIndex = 0;
+	return FileDialogShared(true, ParentWindowHandle, DialogTitle, DefaultPath, DefaultFile, FileTypes, Flags, OutFilenames, DummyFilterIndex);
 }
 
+
+
 bool FDesktopPlatformLinux::OpenDirectoryDialog(const void* ParentWindowHandle, const FString& DialogTitle, const FString& DefaultPath, FString& OutFolderName)
 {
-	unimplemented();
+#if WITH_LINUX_NATIVE_DIALOGS	
+	bool bSuccess;
+	struct UFileDialogHints hints = DEFAULT_UFILEDIALOGHINTS;
+
+	LinuxApplication->SetCapture(NULL);
+
+	hints.Action = UFileDialogActionOpenDirectory;
+	hints.WindowTitle = TCHAR_TO_UTF8(*DialogTitle);
+	hints.InitialDirectory = TCHAR_TO_UTF8(*DefaultPath);
+
+	UFileDialog* dialog = UFileDialog_Create(&hints);
+
+	while(UFileDialog_ProcessEvents(dialog)) {
+		usleep(5000);
+	}
+
+	const UFileDialogResult* result = UFileDialog_Result(dialog);
+
+	if (result)
+	{
+		if (result->count == 1)
+		{
+			OutFolderName = UTF8_TO_TCHAR(result->selection[0]);
+			//OutFolderName = IFileManager::Get().ConvertToRelativePath(*OutFolderName); // @todo (amigo): converting to relative path ends up without ../...
+			FPaths::NormalizeFilename(OutFolderName);
+			bSuccess = true;
+		}
+		else
+		{
+			bSuccess = false;
+		}
+		// Todo like in Windows, normalize files here instead of above
+	}
+	else
+	{
+		bSuccess = false;
+	}
+
+	UFileDialog_Destroy(dialog);
+
+	return bSuccess;
+#else
 	return false;
+#endif // WITH_LINUX_NATIVE_DIALOGS
 }
 
 bool FDesktopPlatformLinux::OpenFontDialog(const void* ParentWindowHandle, FString& OutFontName, float& OutHeight, EFontImportFlags::Type& OutFlags)
@@ -36,48 +89,276 @@ bool FDesktopPlatformLinux::OpenFontDialog(const void* ParentWindowHandle, FStri
 
 bool FDesktopPlatformLinux::OpenLauncher(bool Install, FString CommandLineParams )
 {
-	unimplemented();
-	return false;
+	// TODO: support launcher for realz
+	return true;
 }
 
-bool FDesktopPlatformLinux::FileDialogShared(bool bSave, const void* ParentWindowHandle, const FString& DialogTitle, const FString& DefaultPath, const FString& DefaultFile, const FString& FileTypes, uint32 Flags, TArray<FString>& OutFilenames)
+bool FDesktopPlatformLinux::FileDialogShared(bool bSave, const void* ParentWindowHandle, const FString& DialogTitle, const FString& DefaultPath, const FString& DefaultFile, const FString& FileTypes, uint32 Flags, TArray<FString>& OutFilenames, int32& OutFilterIndex)
 {
-	unimplemented();
+#if WITH_LINUX_NATIVE_DIALOGS
+	FString Filename;
+	bool bSuccess;
+
+	LinuxApplication->SetCapture(NULL);
+
+	UE_LOG(LogDesktopPlatform, Warning, TEXT("FileDialogShared DialogTitle: %s, DefaultPath: %s, DefaultFile: %s, FileTypes: %s, Flags: %d"),
+																		*DialogTitle, *DefaultPath, *DefaultFile, *FileTypes, Flags);
+
+	struct UFileDialogHints hints = DEFAULT_UFILEDIALOGHINTS;
+
+	hints.WindowTitle = TCHAR_TO_UTF8(*DialogTitle);
+
+	// Convert the "|" delimited list of filetypes to NULL delimited then add a second NULL character to indicate the end of the list
+	TCHAR FileTypeStr[MAX_FILETYPES_STR];
+	TCHAR* FileTypesPtr = NULL;
+	const int32 FileTypesLen = FileTypes.Len();
+
+	// Nicely formatted file types for lookup later and suitable to append to filenames without extensions
+	TArray<FString> CleanExtensionList;
+
+	// The strings must be in pairs for windows.
+	// It is formatted as follows: Pair1String1|Pair1String2|Pair2String1|Pair2String2
+	// where the second string in the pair is the extension.  To get the clean extensions we only care about the second string in the pair
+	TArray<FString> UnformattedExtensions;
+	FileTypes.ParseIntoArray( &UnformattedExtensions, TEXT("|"), true );
+	for( int32 ExtensionIndex = 1; ExtensionIndex < UnformattedExtensions.Num(); ExtensionIndex += 2)
+	{
+		const FString& Extension = UnformattedExtensions[ExtensionIndex];
+		// Assume the user typed in an extension or doesn't want one when using the *.* extension. We can't determine what extension they want in that case
+		if( Extension != TEXT("*.*") )
+		{
+			// Add to the clean extension list, first removing the * wildcard from the extension
+			int32 WildCardIndex = Extension.Find( TEXT("*") );
+			CleanExtensionList.Add( WildCardIndex != INDEX_NONE ? Extension.RightChop( WildCardIndex+1 ) : Extension );
+		}
+	}
+
+	if (FileTypesLen > 0 && FileTypesLen - 1 < MAX_FILETYPES_STR)
+	{
+		FileTypesPtr = FileTypeStr;
+		FCString::Strcpy(FileTypeStr, MAX_FILETYPES_STR, *FileTypes);
+
+		TCHAR* Pos = FileTypeStr;
+		while( Pos[0] != 0 )
+		{
+			if ( Pos[0] == '|' )
+			{
+				Pos[0] = 0;
+			}
+
+			Pos++;
+		}
+
+		// Add two trailing NULL characters to indicate the end of the list
+		FileTypeStr[FileTypesLen] = 0;
+		FileTypeStr[FileTypesLen + 1] = 0;
+	}
+
+	//char f_types[4096] = {0,};
+	//FTCHARToUTF8_Convert::Convert(f_types, sizeof(f_types), FileTypeStr, FileTypesLen);
+	hints.NameFilter = TCHAR_TO_UTF8(FileTypeStr);
+
+	//char d_path[4096] = {0,};
+	//FTCHARToUTF8_Convert::Convert(d_path, sizeof(d_path), *DefaultPath, DefaultPath.Len());
+	hints.InitialDirectory = TCHAR_TO_UTF8(*DefaultPath);
+
+	hints.DefaultFile = TCHAR_TO_UTF8(*DefaultFile);
+
+	if (bSave)
+	{
+		hints.Action = UFileDialogActionSave;
+	}
+	else
+	{
+		hints.Action = UFileDialogActionOpen;
+	}
+
+	UFileDialog* dialog = UFileDialog_Create(&hints);
+
+	while(UFileDialog_ProcessEvents(dialog)) {
+		usleep(5000);
+	}
+
+	const UFileDialogResult* result = UFileDialog_Result(dialog);
+
+	if (result)
+	{
+		if (result->count > 1)
+		{
+			// Todo better handling of multi-selects
+			UE_LOG(LogDesktopPlatform, Warning, TEXT("FileDialogShared Selected Files: %d"), result->count);
+			for(int i = 0;i < result->count;++i) {
+				//Filename = FUTF8ToTCHAR(result->selection[i], MAX_FILENAME_STR).Get();
+				Filename = UTF8_TO_TCHAR(result->selection[0]);
+				//new(OutFilenames) FString(Filename);
+				OutFilenames.Add(Filename);
+				Filename = IFileManager::Get().ConvertToRelativePath(*Filename);
+				FPaths::NormalizeFilename(Filename);
+				//UE_LOG(LogDesktopPlatform, Warning, TEXT("FileDialogShared File: %s"), *Filename);
+			}
+			bSuccess = true;
+		}
+		else if (result->count == 1)
+		{
+			//Filename = FUTF8ToTCHAR(result->selection[0], MAX_FILENAME_STR).Get();
+			Filename = UTF8_TO_TCHAR(result->selection[0]);
+			//new(OutFilenames) FString(Filename);
+			OutFilenames.Add(Filename);
+			Filename = IFileManager::Get().ConvertToRelativePath(*Filename);
+			FPaths::NormalizeFilename(Filename);
+			//UE_LOG(LogDesktopPlatform, Warning, TEXT("FileDialogShared File: %s"), *Filename);
+			bSuccess = true;
+		}
+		else
+		{
+			bSuccess = false;
+		}
+		// Todo like in Windows, normalize files here instead of above
+	}
+	else
+	{
+		bSuccess = false;
+	}
+
+	UFileDialog_Destroy(dialog);
+
+	return bSuccess;
+#else
 	return false;
+#endif // WITH_LINUX_NATIVE_DIALOGS
 }
 
 bool FDesktopPlatformLinux::RegisterEngineInstallation(const FString &RootDir, FString &OutIdentifier)
 {
-	return false;
+	bool bRes = false;
+	if (IsValidRootDirectory(RootDir))
+	{
+		FConfigFile ConfigFile;
+		FString ConfigPath = FString(FPlatformProcess::ApplicationSettingsDir()) / FString(TEXT("UnrealEngine")) / FString(TEXT("Install.ini"));
+		ConfigFile.Read(ConfigPath);
+
+		FConfigSection &Section = ConfigFile.FindOrAdd(TEXT("Installations"));
+		OutIdentifier = FGuid::NewGuid().ToString(EGuidFormats::DigitsWithHyphens);
+		Section.AddUnique(*OutIdentifier, RootDir);
+
+		ConfigFile.Dirty = true;
+		ConfigFile.Write(ConfigPath);
+	}
+	return bRes;
 }
 
 void FDesktopPlatformLinux::EnumerateEngineInstallations(TMap<FString, FString> &OutInstallations)
 {
-	unimplemented();
+	EnumerateLauncherEngineInstallations(OutInstallations);
+
+	FString UProjectPath = FString(FPlatformProcess::ApplicationSettingsDir()) / "Unreal.uproject";
+	FArchive* File = IFileManager::Get().CreateFileWriter(*UProjectPath, FILEWRITE_EvenIfReadOnly);
+	if (File)
+	{
+		File->Close();
+		delete File;
+	}
+	else
+	{
+	    SDL_ShowSimpleMessageBox(SDL_MESSAGEBOX_ERROR, "Unable to write to Settings Directory", TCHAR_TO_UTF8(*UProjectPath), NULL);
+	}
+
+	FConfigFile ConfigFile;
+	FString ConfigPath = FString(FPlatformProcess::ApplicationSettingsDir()) / FString(TEXT("UnrealEngine")) / FString(TEXT("Install.ini"));
+	ConfigFile.Read(ConfigPath);
+
+	FConfigSection &Section = ConfigFile.FindOrAdd(TEXT("Installations"));
+
+	// @todo: currently we can enumerate only this installation
+	FString EngineDir = FPaths::EngineDir();
+
+	FString EngineId;
+	const FName* Key = Section.FindKey(EngineDir);
+	if (Key)
+	{
+		EngineId = Key->ToString();
+	}
+	else
+	{
+		if (!OutInstallations.FindKey(EngineDir))
+		{
+			EngineId = FGuid::NewGuid().ToString(EGuidFormats::DigitsWithHyphens);
+			Section.AddUnique(*EngineId, EngineDir);
+			ConfigFile.Dirty = true;
+		}
+	}
+	if (!EngineId.IsEmpty() && !OutInstallations.Find(EngineId))
+	{
+		OutInstallations.Add(EngineId, EngineDir);
+	}
+
+	ConfigFile.Write(ConfigPath);
+
+	IFileManager::Get().Delete(*UProjectPath);
+}
+
+bool FDesktopPlatformLinux::IsSourceDistribution(const FString &RootDir)
+{
+	// Check for the existence of a GenerateProjectFiles.sh file. This allows compatibility with the GitHub 4.0 release.
+	FString GenerateProjectFilesPath = RootDir / TEXT("GenerateProjectFiles.sh");
+	if (IFileManager::Get().FileSize(*GenerateProjectFilesPath) >= 0)
+	{
+		return true;
+	}
+
+	// Otherwise use the default test
+	return FDesktopPlatformBase::IsSourceDistribution(RootDir);
 }
 
 bool FDesktopPlatformLinux::VerifyFileAssociations()
 {
-	unimplemented();
-	return false;
+	STUBBED("FDesktopPlatformLinux::VerifyFileAssociationsg");
+	return true; // for now we are associated
 }
 
 bool FDesktopPlatformLinux::UpdateFileAssociations()
 {
-	unimplemented();
+	//unimplemented();
+	STUBBED("FDesktopPlatformLinux::UpdateFileAssociations");
 	return false;
 }
 
-bool FDesktopPlatformLinux::RunUnrealBuildTool(const FText& Description, const FString& RootDir, const FString& Arguments, FFeedbackContext* Warn)
+bool FDesktopPlatformLinux::OpenProject(const FString &ProjectFileName)
 {
-	unimplemented();
+	// Get the project filename in a native format
+	FString PlatformProjectFileName = ProjectFileName;
+	FPaths::MakePlatformFilename(PlatformProjectFileName);
+
+	STUBBED("FDesktopPlatformLinux::OpenProject");
 	return false;
 }
 
+bool FDesktopPlatformLinux::RunUnrealBuildTool(const FText& Description, const FString& RootDir, const FString& Arguments, FFeedbackContext* Warn)
+{
+	// Get the path to UBT
+	FString UnrealBuildToolPath = RootDir / TEXT("Engine/Binaries/DotNET/UnrealBuildTool.exe");
+	if(IFileManager::Get().FileSize(*UnrealBuildToolPath) < 0)
+	{
+		Warn->Logf(ELogVerbosity::Error, TEXT("Couldn't find UnrealBuildTool at '%s'"), *UnrealBuildToolPath);
+		return false;
+	}
+
+	// Write the output
+	Warn->Logf(TEXT("Running %s %s"), *UnrealBuildToolPath, *Arguments);
+
+	// launch UBT with Mono
+	FString CmdLineParams = FString::Printf(TEXT("\"%s\" %s"), *UnrealBuildToolPath, *Arguments);
+
+	// Spawn it
+	int32 ExitCode = 0;
+	return FFeedbackContextMarkup::PipeProcessOutput(Description, TEXT("/usr/bin/mono"), CmdLineParams, Warn, &ExitCode) && ExitCode == 0;
+}
+
 FFeedbackContext* FDesktopPlatformLinux::GetNativeFeedbackContext()
 {
-	unimplemented();
-	return NULL;
+	//unimplemented();
+	STUBBED("FDesktopPlatformLinux::GetNativeFeedbackContext");
+	return GWarn;
 }
 
 #undef LOCTEXT_NAMESPACE
diff --git a/Engine/Source/Developer/DesktopPlatform/Private/Linux/DesktopPlatformLinux.h b/Engine/Source/Developer/DesktopPlatform/Private/Linux/DesktopPlatformLinux.h
--- a/Engine/Source/Developer/DesktopPlatform/Private/Linux/DesktopPlatformLinux.h
+++ b/Engine/Source/Developer/DesktopPlatform/Private/Linux/DesktopPlatformLinux.h
@@ -18,15 +18,19 @@ public:
 	virtual bool RegisterEngineInstallation(const FString &RootDir, FString &OutIdentifier) override;
 	virtual void EnumerateEngineInstallations(TMap<FString, FString> &OutInstallations) override;
 
+	virtual bool IsSourceDistribution(const FString &RootDir) override;
+
 	virtual bool VerifyFileAssociations() override;
 	virtual bool UpdateFileAssociations() override;
 
+	virtual bool OpenProject(const FString &ProjectFileName) override;
+
 	virtual bool RunUnrealBuildTool(const FText& Description, const FString& RootDir, const FString& Arguments, FFeedbackContext* Warn) override;
 
 	virtual FFeedbackContext* GetNativeFeedbackContext() override;
 
 private:
-	bool FileDialogShared(bool bSave, const void* ParentWindowHandle, const FString& DialogTitle, const FString& DefaultPath, const FString& DefaultFile, const FString& FileTypes, uint32 Flags, TArray<FString>& OutFilenames);
+	bool FileDialogShared(bool bSave, const void* ParentWindowHandle, const FString& DialogTitle, const FString& DefaultPath, const FString& DefaultFile, const FString& FileTypes, uint32 Flags, TArray<FString>& OutFilenames, int32& OutFilterIndex);
 };
 
 typedef FDesktopPlatformLinux FDesktopPlatform;
