add support for detecting installed SDKs at editor start up
add support for showing the various platforms in different menus whether they are properly installed or not
update FPlatformInfo to include information about SDK status, whether or not the platform should be shown on a particular executing platform (Win, Mac, Linux), and whether or not the platform should be shown as part of the binary distribution
fixed platform target modules for IOS, Android, and HTML5 for binary release
TTP341656, TTP342449, TTP342631
#ue4
#platforms

diff --git a/Engine/Source/Developer/DesktopPlatform/Private/PlatformInfo.cpp b/Engine/Source/Developer/DesktopPlatform/Private/PlatformInfo.cpp
--- a/Engine/Source/Developer/DesktopPlatform/Private/PlatformInfo.cpp
+++ b/Engine/Source/Developer/DesktopPlatform/Private/PlatformInfo.cpp
@@ -11,7 +11,7 @@ namespace PlatformInfo
 namespace
 {
 
-FPlatformInfo BuildPlatformInfo(const FName& InPlatformInfoName, const FName& InTargetPlatformName, const FText& InDisplayName, const EPlatformType InPlatformType, const EPlatformFlags::Flags InPlatformFlags, const FPlatformIconPaths& InIconPaths, const FString& InUATCommandLine, const FString& InAutoSDKPath)
+FPlatformInfo BuildPlatformInfo(const FName& InPlatformInfoName, const FName& InTargetPlatformName, const FText& InDisplayName, const EPlatformType InPlatformType, const EPlatformFlags::Flags InPlatformFlags, const FPlatformIconPaths& InIconPaths, const FString& InUATCommandLine, const FString& InAutoSDKPath, EPlatformSDKStatus InStatus, const FString& InTutorial, bool InEnabled, bool InBinary)
 {
 	FPlatformInfo PlatformInfo;
 
@@ -45,43 +45,68 @@ FPlatformInfo BuildPlatformInfo(const FName& InPlatformInfoName, const FName& In
 	PlatformInfo.IconPaths.LargeStyleName  = *FString::Printf(TEXT("Launcher.Platform_%s.Large"), *InPlatformInfoNameString);
 	PlatformInfo.IconPaths.XLargeStyleName = *FString::Printf(TEXT("Launcher.Platform_%s.XLarge"), *InPlatformInfoNameString);
 
+	// SDK data
+	PlatformInfo.SDKStatus = InStatus;
+	PlatformInfo.SDKTutorial = InTutorial;
+
+	// Distribution data
+	PlatformInfo.bEnabledForUse = InEnabled;
+	PlatformInfo.bEnabledInBinary = InBinary;
+
 	return PlatformInfo;
 }
 
+#if PLATFORM_WINDOWS
+static const bool IsAvailableOnWindows = true;
+static const bool IsAvailableOnMac = false;
+static const bool IsAvailableOnLinux = false;
+static const FString IOSTutorial = TEXT("Shared/Tutorials/InstallingiTunesTutorial");
+#elif PLATFORM_MAC
+static const bool IsAvailableOnWindows = false;
+static const bool IsAvailableOnMac = true;
+static const bool IsAvailableOnLinux = false;
+static const FString IOSTutorial = TEXT("Shared/Tutorials/InstallingXCodeTutorial");
+#elif PLATFORM_LINUX
+static const bool IsAvailableOnWindows = false;
+static const bool IsAvailableOnMac = false;
+static const bool IsAvailableOnLinux = true;
+static const FString IOSTutorial = TEXT("Shared/Tutorials/NotYetImplemented");
+#endif
+
 static const FPlatformInfo PlatformInfoArray[] = {
-	//				  PlatformInfoName					TargetPlatformName			DisplayName													PlatformType				PlatformFlags					IconPaths																																			UATCommandLine											AutoSDKPath
-	BuildPlatformInfo(TEXT("WindowsNoEditor"),			TEXT("WindowsNoEditor"),	LOCTEXT("WindowsNoEditor", "Windows"),						EPlatformType::Game,		EPlatformFlags::None,			FPlatformIconPaths(TEXT("Launcher/WindowsTarget/Platform_WindowsNoEditor_24x"), TEXT("Launcher/WindowsTarget/Platform_WindowsNoEditor_128x")),		TEXT(""),												TEXT("")),
-	BuildPlatformInfo(TEXT("WindowsNoEditor_Win32"),	TEXT("WindowsNoEditor"),	LOCTEXT("WindowsNoEditor_Win32", "Windows (32-bit)"),		EPlatformType::Game,		EPlatformFlags::BuildFlavor,	FPlatformIconPaths(TEXT("Launcher/WindowsTarget/Platform_WindowsNoEditor_24x"), TEXT("Launcher/WindowsTarget/Platform_WindowsNoEditor_128x")),		TEXT("-targetplatform=Win32"),							TEXT("")),
-	BuildPlatformInfo(TEXT("WindowsNoEditor_Win64"),	TEXT("WindowsNoEditor"),	LOCTEXT("WindowsNoEditor_Win64", "Windows (64-bit)"),		EPlatformType::Game,		EPlatformFlags::BuildFlavor,	FPlatformIconPaths(TEXT("Launcher/WindowsTarget/Platform_WindowsNoEditor_24x"), TEXT("Launcher/WindowsTarget/Platform_WindowsNoEditor_128x")),		TEXT("-targetplatform=Win64"),							TEXT("")),
-	BuildPlatformInfo(TEXT("Windows"),					TEXT("Windows"),			LOCTEXT("WindowsEditor", "Windows (Editor)"),				EPlatformType::Editor,		EPlatformFlags::None,			FPlatformIconPaths(TEXT("Launcher/WindowsTarget/Platform_Windows_24x"), TEXT("Launcher/WindowsTarget/Platform_Windows_128x")),						TEXT(""),												TEXT("")),
-	BuildPlatformInfo(TEXT("WindowsClient"),			TEXT("WindowsClient"),		LOCTEXT("WindowsClient", "Windows (Client-only)"),			EPlatformType::Client,		EPlatformFlags::None,			FPlatformIconPaths(TEXT("Launcher/WindowsTarget/Platform_Windows_24x"), TEXT("Launcher/WindowsTarget/Platform_Windows_128x")),						TEXT(""),												TEXT("")),
-	BuildPlatformInfo(TEXT("WindowsServer"),			TEXT("WindowsServer"),		LOCTEXT("WindowsServer", "Windows (Dedicated Server)"),		EPlatformType::Server,		EPlatformFlags::None,			FPlatformIconPaths(TEXT("Launcher/WindowsTarget/Platform_WindowsServer_24x"), TEXT("Launcher/WindowsTarget/Platform_WindowsServer_128x")),			TEXT(""),												TEXT("")),
-
-	BuildPlatformInfo(TEXT("MacNoEditor"),				TEXT("MacNoEditor"),		LOCTEXT("MacNoEditor", "Mac"),								EPlatformType::Game,		EPlatformFlags::None,			FPlatformIconPaths(TEXT("Launcher/MacTarget/Platform_Mac_24x"), TEXT("Launcher/MacTarget/Platform_Mac_128x")),										TEXT(""),												TEXT("")),
-	BuildPlatformInfo(TEXT("Mac"),						TEXT("Mac"),				LOCTEXT("MacEditor", "Mac (Editor)"),						EPlatformType::Editor,		EPlatformFlags::None,			FPlatformIconPaths(TEXT("Launcher/MacTarget/Platform_Mac_24x"), TEXT("Launcher/MacTarget/Platform_Mac_128x")),										TEXT(""),												TEXT("")),
-	BuildPlatformInfo(TEXT("MacClient"),				TEXT("MacClient"),			LOCTEXT("MacClient", "Mac (Client-only)"),					EPlatformType::Client,		EPlatformFlags::None,			FPlatformIconPaths(TEXT("Launcher/MacTarget/Platform_Mac_24x"), TEXT("Launcher/MacTarget/Platform_Mac_128x")),										TEXT(""),												TEXT("")),
-	BuildPlatformInfo(TEXT("MacServer"),				TEXT("MacServer"),			LOCTEXT("MacServer", "Mac (Dedicated Server)"),				EPlatformType::Server,		EPlatformFlags::None,			FPlatformIconPaths(TEXT("Launcher/MacTarget/Platform_Mac_24x"), TEXT("Launcher/MacTarget/Platform_Mac_128x")),										TEXT(""),												TEXT("")),
-
-	BuildPlatformInfo(TEXT("LinuxNoEditor"),			TEXT("LinuxNoEditor"),		LOCTEXT("LinuxNoEditor", "Linux"),							EPlatformType::Game,		EPlatformFlags::None,			FPlatformIconPaths(TEXT("Launcher/LinuxTarget/Platform_Linux_24x"), TEXT("Launcher/LinuxTarget/Platform_Linux_128x")),								TEXT(""),												TEXT("Linux_x64")),
-	BuildPlatformInfo(TEXT("Linux"),					TEXT("Linux"),				LOCTEXT("LinuxEditor", "Linux (Editor)"),					EPlatformType::Editor,		EPlatformFlags::None,			FPlatformIconPaths(TEXT("Launcher/LinuxTarget/Platform_Linux_24x"), TEXT("Launcher/LinuxTarget/Platform_Linux_128x")),								TEXT(""),												TEXT("Linux_x64")),
-	BuildPlatformInfo(TEXT("LinuxClient"),				TEXT("LinuxClient"),		LOCTEXT("LinuxClient", "Linux (Client-only)"),				EPlatformType::Client,		EPlatformFlags::None,			FPlatformIconPaths(TEXT("Launcher/LinuxTarget/Platform_Linux_24x"), TEXT("Launcher/LinuxTarget/Platform_Linux_128x")),								TEXT(""),												TEXT("Linux_x64")),
-	BuildPlatformInfo(TEXT("LinuxServer"),				TEXT("LinuxServer"),		LOCTEXT("LinuxServer", "Linux (Dedicated Server)"),			EPlatformType::Server,		EPlatformFlags::None,			FPlatformIconPaths(TEXT("Launcher/LinuxTarget/Platform_Linux_24x"), TEXT("Launcher/LinuxTarget/Platform_Linux_128x")),								TEXT(""),												TEXT("Linux_x64")),
-	
-	BuildPlatformInfo(TEXT("Android"),					TEXT("Android"),			LOCTEXT("Android", "Android"),								EPlatformType::Game,		EPlatformFlags::None,			FPlatformIconPaths(TEXT("Launcher/AndroidTarget/Platform_Android_24x"), TEXT("Launcher/AndroidTarget/Platform_Android_128x")),						TEXT(""),												TEXT("Android")),
-	BuildPlatformInfo(TEXT("Android_All"),				TEXT("Android"),			LOCTEXT("Android_All", "Android (All)"),					EPlatformType::Game,		EPlatformFlags::CookFlavor,		FPlatformIconPaths(TEXT("Launcher/AndroidTarget/Platform_Android_24x"), TEXT("Launcher/AndroidTarget/Platform_Android_128x")),						TEXT(""),												TEXT("Android")),
-	BuildPlatformInfo(TEXT("Android_ATC"),				TEXT("Android_ATC"),		LOCTEXT("Android_ATC", "Android (ATC)"),					EPlatformType::Game,		EPlatformFlags::CookFlavor,		FPlatformIconPaths(TEXT("Launcher/AndroidTarget/Platform_Android_ATC_24x"), TEXT("Launcher/AndroidTarget/Platform_Android_128x")),					TEXT("-targetplatform=Android -cookflavor=ATC"),		TEXT("Android")),
-	BuildPlatformInfo(TEXT("Android_DXT"),				TEXT("Android_DXT"),		LOCTEXT("Android_DXT", "Android (DXT)"),					EPlatformType::Game,		EPlatformFlags::CookFlavor,		FPlatformIconPaths(TEXT("Launcher/AndroidTarget/Platform_Android_DXT_24x"), TEXT("Launcher/AndroidTarget/Platform_Android_128x")),					TEXT("-targetplatform=Android -cookflavor=DXT"),		TEXT("Android")),
-	BuildPlatformInfo(TEXT("Android_ETC1"),				TEXT("Android_ETC1"),		LOCTEXT("Android_ETC1", "Android (ETC1)"),					EPlatformType::Game,		EPlatformFlags::CookFlavor,		FPlatformIconPaths(TEXT("Launcher/AndroidTarget/Platform_Android_ETC1_24x"), TEXT("Launcher/AndroidTarget/Platform_Android_128x")),					TEXT("-targetplatform=Android -cookflavor=ETC1"),		TEXT("Android")),
-	BuildPlatformInfo(TEXT("Android_ETC2"),				TEXT("Android_ETC2"),		LOCTEXT("Android_ETC2", "Android (ETC2)"),					EPlatformType::Game,		EPlatformFlags::CookFlavor,		FPlatformIconPaths(TEXT("Launcher/AndroidTarget/Platform_Android_ETC2_24x"), TEXT("Launcher/AndroidTarget/Platform_Android_128x")),					TEXT("-targetplatform=Android -cookflavor=ETC2"),		TEXT("Android")),
-	BuildPlatformInfo(TEXT("Android_PVRTC"),			TEXT("Android_PVRTC"),		LOCTEXT("Android_PVRTC", "Android (PVRTC)"),				EPlatformType::Game,		EPlatformFlags::CookFlavor,		FPlatformIconPaths(TEXT("Launcher/AndroidTarget/Platform_Android_PVRTC_24x"), TEXT("Launcher/AndroidTarget/Platform_Android_128x")),				TEXT("-targetplatform=Android -cookflavor=PVRTC"),		TEXT("Android")),
-
-	BuildPlatformInfo(TEXT("IOS"),						TEXT("IOS"),				LOCTEXT("IOS", "iOS"),										EPlatformType::Game,		EPlatformFlags::None,			FPlatformIconPaths(TEXT("Launcher/iOSTarget/Platform_iOS_24x"), TEXT("Launcher/iOSTarget/Platform_iOS_128x")),										TEXT(""),												TEXT("")),
-
-	BuildPlatformInfo(TEXT("PS4"),						TEXT("PS4"),				LOCTEXT("PS4", "PlayStation 4"),							EPlatformType::Game,		EPlatformFlags::None,			FPlatformIconPaths(TEXT("Launcher/PS4Target/Platform_PS4_24x"), TEXT("Launcher/PS4Target/Platform_PS4_128x")),										TEXT(""),												TEXT("PS4")),
-
-	BuildPlatformInfo(TEXT("XboxOne"),					TEXT("XboxOne"),			LOCTEXT("XboxOne", "Xbox One"),								EPlatformType::Game,		EPlatformFlags::None,			FPlatformIconPaths(TEXT("Launcher/XboxOneTarget/Platform_XboxOne_24x"), TEXT("Launcher/XboxOneTarget/Platform_XboxOne_128x")),						TEXT(""),												TEXT("XboxOne")),
-
-	BuildPlatformInfo(TEXT("HTML5"),					TEXT("HTML5"),				LOCTEXT("HTML5", "HTML5"),									EPlatformType::Game,		EPlatformFlags::None,			FPlatformIconPaths(TEXT("Launcher/HTML5Target/Platform_HTML5_24x"), TEXT("Launcher/HTML5Target/Platform_HTML5_128x")),								TEXT(""),												TEXT("HTML5")),
+	//				  PlatformInfoName					TargetPlatformName			DisplayName													PlatformType				PlatformFlags					IconPaths																																			UATCommandLine											AutoSDKPath			SDKStatus						SDKTutorial													EnabledForUse								EnabledForBinary
+	BuildPlatformInfo(TEXT("WindowsNoEditor"),			TEXT("WindowsNoEditor"),	LOCTEXT("WindowsNoEditor", "Windows"),						EPlatformType::Game,		EPlatformFlags::None,			FPlatformIconPaths(TEXT("Launcher/WindowsTarget/Platform_WindowsNoEditor_24x"), TEXT("Launcher/WindowsTarget/Platform_WindowsNoEditor_128x")),		TEXT(""),												TEXT(""),			EPlatformSDKStatus::Unknown,	TEXT("Shared/Tutorials/InstallingVisualStudioTutorial"),	IsAvailableOnWindows,						true),
+	BuildPlatformInfo(TEXT("WindowsNoEditor_Win32"),	TEXT("WindowsNoEditor"),	LOCTEXT("WindowsNoEditor_Win32", "Windows (32-bit)"),		EPlatformType::Game,		EPlatformFlags::BuildFlavor,	FPlatformIconPaths(TEXT("Launcher/WindowsTarget/Platform_WindowsNoEditor_24x"), TEXT("Launcher/WindowsTarget/Platform_WindowsNoEditor_128x")),		TEXT("-targetplatform=Win32"),							TEXT(""),			EPlatformSDKStatus::Unknown,	TEXT("Shared/Tutorials/InstallingVisualStudioTutorial"),	IsAvailableOnWindows,						true),
+	BuildPlatformInfo(TEXT("WindowsNoEditor_Win64"),	TEXT("WindowsNoEditor"),	LOCTEXT("WindowsNoEditor_Win64", "Windows (64-bit)"),		EPlatformType::Game,		EPlatformFlags::BuildFlavor,	FPlatformIconPaths(TEXT("Launcher/WindowsTarget/Platform_WindowsNoEditor_24x"), TEXT("Launcher/WindowsTarget/Platform_WindowsNoEditor_128x")),		TEXT("-targetplatform=Win64"),							TEXT(""),			EPlatformSDKStatus::Unknown,	TEXT("Shared/Tutorials/InstallingVisualStudioTutorial"),	IsAvailableOnWindows,						true),
+	BuildPlatformInfo(TEXT("Windows"),					TEXT("Windows"),			LOCTEXT("WindowsEditor", "Windows (Editor)"),				EPlatformType::Editor,		EPlatformFlags::None,			FPlatformIconPaths(TEXT("Launcher/WindowsTarget/Platform_Windows_24x"), TEXT("Launcher/WindowsTarget/Platform_Windows_128x")),						TEXT(""),												TEXT(""),			EPlatformSDKStatus::Unknown,	TEXT("Shared/Tutorials/InstallingVisualStudioTutorial"),	IsAvailableOnWindows,						true),
+	BuildPlatformInfo(TEXT("WindowsClient"),			TEXT("WindowsClient"),		LOCTEXT("WindowsClient", "Windows (Client-only)"),			EPlatformType::Client,		EPlatformFlags::None,			FPlatformIconPaths(TEXT("Launcher/WindowsTarget/Platform_Windows_24x"), TEXT("Launcher/WindowsTarget/Platform_Windows_128x")),						TEXT(""),												TEXT(""),			EPlatformSDKStatus::Unknown,	TEXT("Shared/Tutorials/InstallingVisualStudioTutorial"),	IsAvailableOnWindows,						true),
+	BuildPlatformInfo(TEXT("WindowsServer"),			TEXT("WindowsServer"),		LOCTEXT("WindowsServer", "Windows (Dedicated Server)"),		EPlatformType::Server,		EPlatformFlags::None,			FPlatformIconPaths(TEXT("Launcher/WindowsTarget/Platform_WindowsServer_24x"), TEXT("Launcher/WindowsTarget/Platform_WindowsServer_128x")),			TEXT(""),												TEXT(""),			EPlatformSDKStatus::Unknown,	TEXT("Shared/Tutorials/InstallingVisualStudioTutorial"),	IsAvailableOnWindows,						true),
+
+	BuildPlatformInfo(TEXT("MacNoEditor"),				TEXT("MacNoEditor"),		LOCTEXT("MacNoEditor", "Mac"),								EPlatformType::Game,		EPlatformFlags::None,			FPlatformIconPaths(TEXT("Launcher/MacTarget/Platform_Mac_24x"), TEXT("Launcher/MacTarget/Platform_Mac_128x")),										TEXT(""),												TEXT(""),			EPlatformSDKStatus::Unknown,	TEXT("Shared/Tutorials/InstallingXCodeTutorial"),			IsAvailableOnMac,							true),
+	BuildPlatformInfo(TEXT("Mac"),						TEXT("Mac"),				LOCTEXT("MacEditor", "Mac (Editor)"),						EPlatformType::Editor,		EPlatformFlags::None,			FPlatformIconPaths(TEXT("Launcher/MacTarget/Platform_Mac_24x"), TEXT("Launcher/MacTarget/Platform_Mac_128x")),										TEXT(""),												TEXT(""),			EPlatformSDKStatus::Unknown,	TEXT("Shared/Tutorials/InstallingXCodeTutorial"),			IsAvailableOnMac,							true),
+	BuildPlatformInfo(TEXT("MacClient"),				TEXT("MacClient"),			LOCTEXT("MacClient", "Mac (Client-only)"),					EPlatformType::Client,		EPlatformFlags::None,			FPlatformIconPaths(TEXT("Launcher/MacTarget/Platform_Mac_24x"), TEXT("Launcher/MacTarget/Platform_Mac_128x")),										TEXT(""),												TEXT(""),			EPlatformSDKStatus::Unknown,	TEXT("Shared/Tutorials/InstallingXCodeTutorial"),			IsAvailableOnMac,							true),
+	BuildPlatformInfo(TEXT("MacServer"),				TEXT("MacServer"),			LOCTEXT("MacServer", "Mac (Dedicated Server)"),				EPlatformType::Server,		EPlatformFlags::None,			FPlatformIconPaths(TEXT("Launcher/MacTarget/Platform_Mac_24x"), TEXT("Launcher/MacTarget/Platform_Mac_128x")),										TEXT(""),												TEXT(""),			EPlatformSDKStatus::Unknown,	TEXT("Shared/Tutorials/InstallingXCodeTutorial"),			IsAvailableOnMac,							true),
+
+	BuildPlatformInfo(TEXT("LinuxNoEditor"),			TEXT("LinuxNoEditor"),		LOCTEXT("LinuxNoEditor", "Linux"),							EPlatformType::Game,		EPlatformFlags::None,			FPlatformIconPaths(TEXT("Launcher/LinuxTarget/Platform_Linux_24x"), TEXT("Launcher/LinuxTarget/Platform_Linux_128x")),								TEXT(""),												TEXT("Linux_x64"),	EPlatformSDKStatus::Unknown,	TEXT(""),													IsAvailableOnLinux || IsAvailableOnWindows,	false),	
+	BuildPlatformInfo(TEXT("Linux"),					TEXT("Linux"),				LOCTEXT("LinuxEditor", "Linux (Editor)"),					EPlatformType::Editor,		EPlatformFlags::None,			FPlatformIconPaths(TEXT("Launcher/LinuxTarget/Platform_Linux_24x"), TEXT("Launcher/LinuxTarget/Platform_Linux_128x")),								TEXT(""),												TEXT("Linux_x64"),	EPlatformSDKStatus::Unknown,	TEXT(""),													IsAvailableOnLinux,							false),
+	BuildPlatformInfo(TEXT("LinuxClient"),				TEXT("LinuxClient"),		LOCTEXT("LinuxClient", "Linux (Client-only)"),				EPlatformType::Client,		EPlatformFlags::None,			FPlatformIconPaths(TEXT("Launcher/LinuxTarget/Platform_Linux_24x"), TEXT("Launcher/LinuxTarget/Platform_Linux_128x")),								TEXT(""),												TEXT("Linux_x64"),	EPlatformSDKStatus::Unknown,	TEXT(""),													IsAvailableOnLinux || IsAvailableOnWindows,	false),
+	BuildPlatformInfo(TEXT("LinuxServer"),				TEXT("LinuxServer"),		LOCTEXT("LinuxServer", "Linux (Dedicated Server)"),			EPlatformType::Server,		EPlatformFlags::None,			FPlatformIconPaths(TEXT("Launcher/LinuxTarget/Platform_Linux_24x"), TEXT("Launcher/LinuxTarget/Platform_Linux_128x")),								TEXT(""),												TEXT("Linux_x64"),	EPlatformSDKStatus::Unknown,	TEXT(""),													IsAvailableOnLinux || IsAvailableOnWindows,	false),
+
+	BuildPlatformInfo(TEXT("IOS"),						TEXT("IOS"),				LOCTEXT("IOS", "iOS"),										EPlatformType::Game,		EPlatformFlags::None,			FPlatformIconPaths(TEXT("Launcher/iOSTarget/Platform_iOS_24x"), TEXT("Launcher/iOSTarget/Platform_iOS_128x")),										TEXT(""),												TEXT(""),			EPlatformSDKStatus::Unknown,	IOSTutorial,												IsAvailableOnWindows || IsAvailableOnMac,	true),
+
+	BuildPlatformInfo(TEXT("Android"),					TEXT("Android"),			LOCTEXT("Android", "Android"),								EPlatformType::Game,		EPlatformFlags::None,			FPlatformIconPaths(TEXT("Launcher/AndroidTarget/Platform_Android_24x"), TEXT("Launcher/AndroidTarget/Platform_Android_128x")),						TEXT(""),												TEXT("Android"),	EPlatformSDKStatus::Unknown,	TEXT("Shared/Tutorials/SettingUpAndroidTutorial"),			IsAvailableOnWindows,						true),
+	BuildPlatformInfo(TEXT("Android_All"),				TEXT("Android"),			LOCTEXT("Android_All", "Android (All)"),					EPlatformType::Game,		EPlatformFlags::CookFlavor,		FPlatformIconPaths(TEXT("Launcher/AndroidTarget/Platform_Android_24x"), TEXT("Launcher/AndroidTarget/Platform_Android_128x")),						TEXT(""),												TEXT("Android"),	EPlatformSDKStatus::Unknown,	TEXT("Shared/Tutorials/SettingUpAndroidTutorial"),			IsAvailableOnWindows,						true),
+	BuildPlatformInfo(TEXT("Android_ATC"),				TEXT("Android_ATC"),		LOCTEXT("Android_ATC", "Android (ATC)"),					EPlatformType::Game,		EPlatformFlags::CookFlavor,		FPlatformIconPaths(TEXT("Launcher/AndroidTarget/Platform_Android_ATC_24x"), TEXT("Launcher/AndroidTarget/Platform_Android_128x")),					TEXT("-targetplatform=Android -cookflavor=ATC"),		TEXT("Android"),	EPlatformSDKStatus::Unknown,	TEXT("Shared/Tutorials/SettingUpAndroidTutorial"),			IsAvailableOnWindows,						true),
+	BuildPlatformInfo(TEXT("Android_DXT"),				TEXT("Android_DXT"),		LOCTEXT("Android_DXT", "Android (DXT)"),					EPlatformType::Game,		EPlatformFlags::CookFlavor,		FPlatformIconPaths(TEXT("Launcher/AndroidTarget/Platform_Android_DXT_24x"), TEXT("Launcher/AndroidTarget/Platform_Android_128x")),					TEXT("-targetplatform=Android -cookflavor=DXT"),		TEXT("Android"),	EPlatformSDKStatus::Unknown,	TEXT("Shared/Tutorials/SettingUpAndroidTutorial"),			IsAvailableOnWindows,						true),
+	BuildPlatformInfo(TEXT("Android_ETC1"),				TEXT("Android_ETC1"),		LOCTEXT("Android_ETC1", "Android (ETC1)"),					EPlatformType::Game,		EPlatformFlags::CookFlavor,		FPlatformIconPaths(TEXT("Launcher/AndroidTarget/Platform_Android_ETC1_24x"), TEXT("Launcher/AndroidTarget/Platform_Android_128x")),					TEXT("-targetplatform=Android -cookflavor=ETC1"),		TEXT("Android"),	EPlatformSDKStatus::Unknown,	TEXT("Shared/Tutorials/SettingUpAndroidTutorial"),			IsAvailableOnWindows,						true),
+	BuildPlatformInfo(TEXT("Android_ETC2"),				TEXT("Android_ETC2"),		LOCTEXT("Android_ETC2", "Android (ETC2)"),					EPlatformType::Game,		EPlatformFlags::CookFlavor,		FPlatformIconPaths(TEXT("Launcher/AndroidTarget/Platform_Android_ETC2_24x"), TEXT("Launcher/AndroidTarget/Platform_Android_128x")),					TEXT("-targetplatform=Android -cookflavor=ETC2"),		TEXT("Android"),	EPlatformSDKStatus::Unknown,	TEXT("Shared/Tutorials/SettingUpAndroidTutorial"),			IsAvailableOnWindows,						true),
+	BuildPlatformInfo(TEXT("Android_PVRTC"),			TEXT("Android_PVRTC"),		LOCTEXT("Android_PVRTC", "Android (PVRTC)"),				EPlatformType::Game,		EPlatformFlags::CookFlavor,		FPlatformIconPaths(TEXT("Launcher/AndroidTarget/Platform_Android_PVRTC_24x"), TEXT("Launcher/AndroidTarget/Platform_Android_128x")),				TEXT("-targetplatform=Android -cookflavor=PVRTC"),		TEXT("Android"),	EPlatformSDKStatus::Unknown,	TEXT("Shared/Tutorials/SettingUpAndroidTutorial"),			IsAvailableOnWindows,						true),
+
+	BuildPlatformInfo(TEXT("HTML5"),					TEXT("HTML5"),				LOCTEXT("HTML5", "HTML5"),									EPlatformType::Game,		EPlatformFlags::None,			FPlatformIconPaths(TEXT("Launcher/HTML5Target/Platform_HTML5_24x"), TEXT("Launcher/HTML5Target/Platform_HTML5_128x")),								TEXT(""),												TEXT("HTML5"),		EPlatformSDKStatus::Unknown,	TEXT("Platforms/HTML5/GettingStarted"),						IsAvailableOnWindows,						false),
+
+	BuildPlatformInfo(TEXT("PS4"),						TEXT("PS4"),				LOCTEXT("PS4", "PlayStation 4"),							EPlatformType::Game,		EPlatformFlags::None,			FPlatformIconPaths(TEXT("Launcher/PS4Target/Platform_PS4_24x"), TEXT("Launcher/PS4Target/Platform_PS4_128x")),										TEXT(""),												TEXT("PS4"),		EPlatformSDKStatus::Unknown,	TEXT("Platforms/PS4/GettingStarted"),						IsAvailableOnWindows,						false),
+
+	BuildPlatformInfo(TEXT("XboxOne"),					TEXT("XboxOne"),			LOCTEXT("XboxOne", "Xbox One"),								EPlatformType::Game,		EPlatformFlags::None,			FPlatformIconPaths(TEXT("Launcher/XboxOneTarget/Platform_XboxOne_24x"), TEXT("Launcher/XboxOneTarget/Platform_XboxOne_128x")),						TEXT(""),												TEXT("XboxOne"),	EPlatformSDKStatus::Unknown,	TEXT("Platforms/XboxOne/GettingStarted"),					IsAvailableOnWindows,						false),
 
 	//BuildPlatformInfo(TEXT("WinRT"),					TEXT("WinRT"),				LOCTEXT("WinRT", "Windows RT"),								EPlatformType::Game,		EPlatformFlags::None,			FPlatformIconPaths(TEXT("Launcher/WindowsTarget/Platform_WindowsNoEditor_24x"), TEXT("Launcher/WindowsTarget/Platform_WindowsNoEditor_128x")),		TEXT("")),
 };
@@ -112,6 +137,17 @@ const FPlatformInfo* GetPlatformInfoArray(int32& OutNumPlatforms)
 	return PlatformInfoArray;
 }
 
+void UpdatePlatformSDKStatus(FString InPlatformName, EPlatformSDKStatus InStatus)
+{
+	for(const FPlatformInfo& PlatformInfo : PlatformInfoArray)
+	{
+		if(PlatformInfo.VanillaPlatformName == FName(*InPlatformName))
+		{
+			const_cast<FPlatformInfo&>(PlatformInfo).SDKStatus = InStatus;
+		}
+	}
+}
+
 FPlatformEnumerator EnumeratePlatformInfoArray()
 {
 	return FPlatformEnumerator(PlatformInfoArray, ARRAY_COUNT(PlatformInfoArray));
diff --git a/Engine/Source/Developer/DesktopPlatform/Public/PlatformInfo.h b/Engine/Source/Developer/DesktopPlatform/Public/PlatformInfo.h
--- a/Engine/Source/Developer/DesktopPlatform/Public/PlatformInfo.h
+++ b/Engine/Source/Developer/DesktopPlatform/Public/PlatformInfo.h
@@ -63,6 +63,18 @@ namespace PlatformInfo
 		};
 	}
 
+	enum class EPlatformSDKStatus : uint8
+	{
+		/** SDK status is unknown */
+		Unknown,
+
+		/** SDK is installed */
+		Installed,
+
+		/** SDK is not installed */
+		NotInstalled,
+	};
+
 	/** Information about where to find the platform icons (for use by FEditorStyle) */
 	struct FPlatformIconPaths
 	{
@@ -134,6 +146,18 @@ namespace PlatformInfo
 		/** Path under CarefullyRedist for the SDK.  FString so case sensitive platforms don't get messed up by a pre-existing FName of a different casing. */
 		FString AutoSDKPath;
 
+		/** Whether or not this platform SDK has been properly installed */
+		EPlatformSDKStatus SDKStatus;
+
+		/** Tutorial path for tutorial to install SDK */
+		FString SDKTutorial;
+
+		/** Enabled for use */
+		bool bEnabledForUse;
+
+		/** Enabled in binary distribution */
+		bool bEnabledInBinary;
+
 		/** Returns true if this platform is vanilla */
 		FORCEINLINE bool IsVanilla() const
 		{
@@ -285,4 +309,12 @@ namespace PlatformInfo
 	 * @return An array of vanilla platforms, potentially containing flavors
 	 */
 	DESKTOPPLATFORM_API TArray<FVanillaPlatformEntry> BuildPlatformHierarchy(const EPlatformFilter InFilter);
+
+	/**
+	 * Get an array of all the platforms we know about
+	 * @param OutNumPlatforms - The number of platforms in the array
+	 * @return The pointer to the start of the platform array
+	 */
+	DESKTOPPLATFORM_API void UpdatePlatformSDKStatus(FString InPlatformName, EPlatformSDKStatus InStatus);
+
 }
diff --git a/Engine/Source/Developer/TargetPlatform/Private/TargetPlatformManagerModule.cpp b/Engine/Source/Developer/TargetPlatform/Private/TargetPlatformManagerModule.cpp
--- a/Engine/Source/Developer/TargetPlatform/Private/TargetPlatformManagerModule.cpp
+++ b/Engine/Source/Developer/TargetPlatform/Private/TargetPlatformManagerModule.cpp
@@ -55,6 +55,7 @@ public:
 		}
 #endif
 
+		SetupSDKStatus();
 		GetTargetPlatforms();
 		GetActiveTargetPlatforms();
 		GetAudioFormats();
@@ -80,6 +81,7 @@ public:
 	{
 		bForceCacheUpdate = true;
 
+		SetupSDKStatus();
 		GetTargetPlatforms();
 		GetActiveTargetPlatforms();
 		GetAudioFormats();
@@ -782,6 +784,69 @@ protected:
 #endif
 	}
 
+	bool SetupSDKStatus()
+	{
+		// run UBT with -validate -allplatforms and read the output
+		FString CmdExe = FPaths::ConvertRelativePathToFull(FPaths::EngineDir() / TEXT("Binaries/DotNet/UnrealBuildTool.exe"));
+		FString CommandLine = TEXT("-validateplatform -allplatforms");
+		TSharedPtr<FMonitoredProcess> UBTProcess = MakeShareable(new FMonitoredProcess(CmdExe, CommandLine, true));
+		UBTProcess->OnOutput().BindStatic(&FTargetPlatformManagerModule::OnStatusOutput);
+		SDKStatusMessage = TEXT("");
+		UBTProcess->Launch();
+		while(UBTProcess->IsRunning())
+		{
+			FPlatformProcess::Sleep(0.01f);
+		}
+
+		TArray<FString> LineArray;
+		SDKStatusMessage.ParseIntoArrayLines(&LineArray);
+		for (int Index = 0; Index < LineArray.Num(); ++Index)
+		{
+			TArray<FString> PlatArray;
+			LineArray[Index].ParseIntoArrayWS(&PlatArray);
+			if (PlatArray.Num() == 3)
+			{
+				PlatformInfo::EPlatformSDKStatus Status = PlatArray[2] == TEXT("VALID") ? PlatformInfo::EPlatformSDKStatus::Installed : PlatformInfo::EPlatformSDKStatus::NotInstalled;
+				FString PlatformName = PlatArray[1];
+				if (PlatformName == TEXT("Win32") || PlatformName == TEXT("Win64"))
+				{
+					PlatformName = TEXT("Windows");
+					PlatformInfo::UpdatePlatformSDKStatus(PlatformName, Status);
+					PlatformName = TEXT("WindowsNoEditor");
+					PlatformInfo::UpdatePlatformSDKStatus(PlatformName, Status);
+					PlatformName = TEXT("WindowsClient");
+					PlatformInfo::UpdatePlatformSDKStatus(PlatformName, Status);
+					PlatformName = TEXT("WindowsServer");
+					PlatformInfo::UpdatePlatformSDKStatus(PlatformName, Status);
+				}
+				else if (PlatformName == TEXT("Mac"))
+				{
+					PlatformInfo::UpdatePlatformSDKStatus(PlatformName, Status);
+					PlatformName = TEXT("MacNoEditor");
+					PlatformInfo::UpdatePlatformSDKStatus(PlatformName, Status);
+					PlatformName = TEXT("MacClient");
+					PlatformInfo::UpdatePlatformSDKStatus(PlatformName, Status);
+					PlatformName = TEXT("MacServer");
+					PlatformInfo::UpdatePlatformSDKStatus(PlatformName, Status);
+				}
+				else if (PlatformName == TEXT("Linux"))
+				{
+					PlatformInfo::UpdatePlatformSDKStatus(PlatformName, Status);
+					PlatformName = TEXT("LinuxNoEditor");
+					PlatformInfo::UpdatePlatformSDKStatus(PlatformName, Status);
+					PlatformName = TEXT("LinuxClient");
+					PlatformInfo::UpdatePlatformSDKStatus(PlatformName, Status);
+					PlatformName = TEXT("LinuxServer");
+					PlatformInfo::UpdatePlatformSDKStatus(PlatformName, Status);
+				}
+				else
+				{
+					PlatformInfo::UpdatePlatformSDKStatus(PlatformName, Status);
+				}
+			}
+		}
+		return true;
+	}
 
 private:
 
@@ -793,6 +858,12 @@ private:
 		}
 	}
 	
+	static FString SDKStatusMessage;
+	static void OnStatusOutput(FString Message)
+	{
+		SDKStatusMessage += Message;
+	}
+
 	// If true we should build formats that are actually required for use by the runtime. 
 	// This happens for an ordinary editor run and more specifically whenever there is no
 	// TargetPlatform= on the command line.
@@ -811,5 +882,6 @@ private:
 #endif
 };
 
+FString FTargetPlatformManagerModule::SDKStatusMessage = TEXT("");
 
 IMPLEMENT_MODULE(FTargetPlatformManagerModule, TargetPlatform);
diff --git a/Engine/Source/Developer/iOS/IOSTargetPlatform/Private/IOSTargetPlatform.cpp b/Engine/Source/Developer/iOS/IOSTargetPlatform/Private/IOSTargetPlatform.cpp
--- a/Engine/Source/Developer/iOS/IOSTargetPlatform/Private/IOSTargetPlatform.cpp
+++ b/Engine/Source/Developer/iOS/IOSTargetPlatform/Private/IOSTargetPlatform.cpp
@@ -123,7 +123,7 @@ int FIOSTargetPlatform::IsReadyToBuild(const FString& ProjectPath, bool bProject
 	// shell to IPP and get the status of the provision and cert
 	FString CmdExe = FPaths::ConvertRelativePathToFull(FPaths::EngineDir() / TEXT("Binaries/DotNet/IOS/IPhonePackager.exe"));
 	FString CommandLine = FString::Printf(TEXT("Validate Engine -project %s"), *ProjectPath);
-	TSharedPtr<FMonitoredProcess> IPPProcess = MakeShareable(new FMonitoredProcess(CmdExe, CommandLine, false));
+	TSharedPtr<FMonitoredProcess> IPPProcess = MakeShareable(new FMonitoredProcess(CmdExe, CommandLine, true));
 	IPPProcess->Launch();
 	while(IPPProcess->IsRunning())
 	{
diff --git a/Engine/Source/Editor/MainFrame/Private/Frame/MainFrameActions.cpp b/Engine/Source/Editor/MainFrame/Private/Frame/MainFrameActions.cpp
--- a/Engine/Source/Editor/MainFrame/Private/Frame/MainFrameActions.cpp
+++ b/Engine/Source/Editor/MainFrame/Private/Frame/MainFrameActions.cpp
@@ -413,6 +413,13 @@ void FMainFrameActionCallbacks::CookContent( const FName InPlatformInfoName )
 		return;
 	}
 
+	if (PlatformInfo->SDKStatus == PlatformInfo::EPlatformSDKStatus::NotInstalled)
+	{
+		IMainFrameModule& MainFrameModule = FModuleManager::GetModuleChecked<IMainFrameModule>(TEXT("MainFrame"));
+		MainFrameModule.BroadcastMainFrameSDKNotInstalled(PlatformInfo->TargetPlatformName.ToString(), PlatformInfo->SDKTutorial);
+		return;
+	}
+
 	if (PlatformInfo->TargetPlatformName == FName("MacNoEditor"))
 	{
 		OptionalParams += TEXT(" -targetplatform=Mac");
@@ -474,6 +481,13 @@ void FMainFrameActionCallbacks::PackageProject( const FName InPlatformInfoName )
 	const PlatformInfo::FPlatformInfo* const PlatformInfo = PlatformInfo::FindPlatformInfo(InPlatformInfoName);
 	check(PlatformInfo);
 
+	if (PlatformInfo->SDKStatus == PlatformInfo::EPlatformSDKStatus::NotInstalled)
+	{
+		IMainFrameModule& MainFrameModule = FModuleManager::GetModuleChecked<IMainFrameModule>(TEXT("MainFrame"));
+		MainFrameModule.BroadcastMainFrameSDKNotInstalled(PlatformInfo->TargetPlatformName.ToString(), PlatformInfo->SDKTutorial);
+		return;
+	}
+
 	{
 		const ITargetPlatform* const Platform = GetTargetPlatformManager()->FindTargetPlatform(PlatformInfo->TargetPlatformName.ToString());
 		if (Platform)
@@ -624,10 +638,10 @@ void FMainFrameActionCallbacks::PackageProject( const FName InPlatformInfoName )
 	CreateUatTask(CommandLine, PlatformInfo->DisplayName, LOCTEXT("PackagingProjectTaskName", "Packaging project"), LOCTEXT("PackagingTaskName", "Packaging"), FEditorStyle::GetBrush(TEXT("MainFrame.PackageProject")));
 }
 
-bool FMainFrameActionCallbacks::PackageProjectCanExecute( const FName PlatformInfoName, bool IsImplemented )
+bool FMainFrameActionCallbacks::PackageProjectCanExecute( const FName PlatformInfoName )
 {
 	// For a binary Rocket build, Development is ALWAYS Win64, and Shipping is ALWAYS Win32.
-	if(IsImplemented && FRocketSupport::IsRocket())
+	if(FRocketSupport::IsRocket())
 	{
 		const EProjectPackagingBuildConfigurations BuildConfiguration = GetDefault<UProjectPackagingSettings>()->BuildConfiguration;
 		switch(BuildConfiguration)
@@ -644,7 +658,7 @@ bool FMainFrameActionCallbacks::PackageProjectCanExecute( const FName PlatformIn
 		}
 	}
 
-	return IsImplemented;
+	return true;
 }
 
 void FMainFrameActionCallbacks::RefreshCodeProject()
diff --git a/Engine/Source/Editor/MainFrame/Private/Frame/MainFrameActions.h b/Engine/Source/Editor/MainFrame/Private/Frame/MainFrameActions.h
--- a/Engine/Source/Editor/MainFrame/Private/Frame/MainFrameActions.h
+++ b/Engine/Source/Editor/MainFrame/Private/Frame/MainFrameActions.h
@@ -167,7 +167,7 @@ public:
 	static void PackageProject( const FName InPlatformInfoName );
 
 	/** Checks whether a menu action for packaging the project can execute. */
-	static bool PackageProjectCanExecute( const FName PlatformInfoName, bool IsImplemented );
+	static bool PackageProjectCanExecute( const FName PlatformInfoName );
 
 	/** Refresh the project in the current IDE */
 	static void RefreshCodeProject();
diff --git a/Engine/Source/Editor/MainFrame/Private/Menus/CookContentMenu.h b/Engine/Source/Editor/MainFrame/Private/Menus/CookContentMenu.h
--- a/Engine/Source/Editor/MainFrame/Private/Menus/CookContentMenu.h
+++ b/Engine/Source/Editor/MainFrame/Private/Menus/CookContentMenu.h
@@ -41,17 +41,17 @@ public:
 			check(VanillaPlatform.PlatformInfo->IsVanilla());
 
 			// Only care about game targets
-			if (VanillaPlatform.PlatformInfo->PlatformType != PlatformInfo::EPlatformType::Game)
+			if (VanillaPlatform.PlatformInfo->PlatformType != PlatformInfo::EPlatformType::Game || !VanillaPlatform.PlatformInfo->bEnabledForUse || (!VanillaPlatform.PlatformInfo->bEnabledInBinary && FRocketSupport::IsRocket()))
 			{
 				continue;
 			}
 
-			// Make sure we're able to run this platform
-			ITargetPlatform* const Platform = GetTargetPlatformManager()->FindTargetPlatform(VanillaPlatform.PlatformInfo->TargetPlatformName.ToString());
+			// We now make sure we're able to run this platform at the command stage so we can fire off SDK tutorials
+/*			ITargetPlatform* const Platform = GetTargetPlatformManager()->FindTargetPlatform(VanillaPlatform.PlatformInfo->TargetPlatformName.ToString());
 			if (!Platform)
 			{
 				continue;
-			}
+			}*/
 
 			if(VanillaPlatform.PlatformFlavors.Num())
 			{
diff --git a/Engine/Source/Editor/MainFrame/Private/Menus/PackageProjectMenu.h b/Engine/Source/Editor/MainFrame/Private/Menus/PackageProjectMenu.h
--- a/Engine/Source/Editor/MainFrame/Private/Menus/PackageProjectMenu.h
+++ b/Engine/Source/Editor/MainFrame/Private/Menus/PackageProjectMenu.h
@@ -40,17 +40,17 @@ public:
 			check(VanillaPlatform.PlatformInfo->IsVanilla());
 
 			// Only care about game targets
-			if (VanillaPlatform.PlatformInfo->PlatformType != PlatformInfo::EPlatformType::Game)
+			if (VanillaPlatform.PlatformInfo->PlatformType != PlatformInfo::EPlatformType::Game || !VanillaPlatform.PlatformInfo->bEnabledForUse || (!VanillaPlatform.PlatformInfo->bEnabledInBinary && FRocketSupport::IsRocket()))
 			{
 				continue;
 			}
 
 			// Make sure we're able to run this platform
-			ITargetPlatform* const Platform = GetTargetPlatformManager()->FindTargetPlatform(VanillaPlatform.PlatformInfo->TargetPlatformName.ToString());
+/*			ITargetPlatform* const Platform = GetTargetPlatformManager()->FindTargetPlatform(VanillaPlatform.PlatformInfo->TargetPlatformName.ToString());
 			if (!Platform)
 			{
 				continue;
-			}
+			}*/
 
 			if (VanillaPlatform.PlatformFlavors.Num())
 			{
@@ -91,11 +91,9 @@ protected:
 	{
 		IProjectTargetPlatformEditorModule& ProjectTargetPlatformEditorModule = FModuleManager::LoadModuleChecked<IProjectTargetPlatformEditorModule>("ProjectTargetPlatformEditor");
 
-		ITargetPlatform* const Platform = GetTargetPlatformManager()->FindTargetPlatform(PlatformInfo.TargetPlatformName.ToString());
-
 		FUIAction Action(
 			FExecuteAction::CreateStatic(&FMainFrameActionCallbacks::PackageProject, PlatformInfo.PlatformInfoName),
-			FCanExecuteAction::CreateStatic(&FMainFrameActionCallbacks::PackageProjectCanExecute, PlatformInfo.PlatformInfoName, Platform && Platform->SupportsFeature(ETargetPlatformFeatures::Packaging))
+			FCanExecuteAction::CreateStatic(&FMainFrameActionCallbacks::PackageProjectCanExecute, PlatformInfo.PlatformInfoName)
 			);
 
 		// ... generate tooltip text
diff --git a/Engine/Source/Editor/UnrealEd/Private/Kismet2/DebuggerCommands.cpp b/Engine/Source/Editor/UnrealEd/Private/Kismet2/DebuggerCommands.cpp
--- a/Engine/Source/Editor/UnrealEd/Private/Kismet2/DebuggerCommands.cpp
+++ b/Engine/Source/Editor/UnrealEd/Private/Kismet2/DebuggerCommands.cpp
@@ -62,7 +62,11 @@ public:
 	static bool HandleLaunchOnDeviceActionCanExecute(FString DevicedId);
 	static bool HandleLaunchOnDeviceActionIsChecked(FString DeviceId);
 
-	static void HandleShowSDKTutorial(ITargetPlatform* Platform);
+	// No Device
+	static void HandleNoDeviceFoundActionExecute() {}
+	static bool HandleNoDeviceFoundActionCanExecute() { return false; }
+
+	static void HandleShowSDKTutorial(FString PlatformName, FString NotInstalledDocLink);
 
 	static void RepeatLastLaunch_Clicked();
 	static bool RepeatLastLaunch_CanExecute();
@@ -558,13 +562,13 @@ TSharedRef< SWidget > FPlayWorldCommands::GenerateLaunchMenuContent( TSharedRef<
 	const bool bShouldCloseWindowAfterMenuSelection = true;
 	FMenuBuilder MenuBuilder(bShouldCloseWindowAfterMenuSelection, InCommandList);
 
-	TArray<ITargetPlatform*> Platforms = GetTargetPlatformManager()->GetTargetPlatforms();
+/*	TArray<ITargetPlatform*> Platforms = GetTargetPlatformManager()->GetTargetPlatforms();
 
 	// Sort the platforms by display name so that they appear in the same order as the other menus (like "Package Project" and "Cook Project")
 	Platforms.Sort([](ITargetPlatform& One, ITargetPlatform& Two) -> bool
 	{
 		return One.DisplayName().CompareTo(Two.DisplayName()) < 0;
-	});
+	});*/
 
 	// shared devices section
 	TSharedPtr<ITargetDeviceServicesModule> TargetDeviceServicesModule = StaticCastSharedPtr<ITargetDeviceServicesModule>(FModuleManager::Get().LoadModule(TEXT("TargetDeviceServices")));
@@ -572,88 +576,127 @@ TSharedRef< SWidget > FPlayWorldCommands::GenerateLaunchMenuContent( TSharedRef<
 
 	TArray<FString> PlatformsToMaybeInstallLinksFor;
 	PlatformsToMaybeInstallLinksFor.Add(TEXT("Android"));
-	TArray<ITargetPlatform*> PlatformsToAddInstallLinksFor;
+	PlatformsToMaybeInstallLinksFor.Add(TEXT("IOS"));
+	PlatformsToMaybeInstallLinksFor.Add(TEXT("Linux"));
+	TArray<FString> PlatformsWithNoDevices;
+	TArray<PlatformInfo::FPlatformInfo> PlatformsToAddInstallLinksFor;
 
 	MenuBuilder.BeginSection("LevelEditorLaunchDevices", LOCTEXT("LaunchButtonDevicesSection", "Devices"));
 	{
-		for(ITargetPlatform* Platform : Platforms)
+		for(const PlatformInfo::FPlatformInfo& PlatformInfo : PlatformInfo::EnumeratePlatformInfoArray())
 		{
-			const PlatformInfo::FPlatformInfo& PlatformInfo = Platform->GetPlatformInfo();
-
 			// for the Editor we are only interested in launching standalone games
-			if (PlatformInfo.PlatformType != PlatformInfo::EPlatformType::Game)
+			if (PlatformInfo.PlatformType != PlatformInfo::EPlatformType::Game || !PlatformInfo.bEnabledForUse || (!PlatformInfo.bEnabledInBinary && FRocketSupport::IsRocket()))
 			{
 				continue;
 			}
 
-			// for each platform...
-			TArray<ITargetDeviceProxyPtr> DeviceProxies;
-			TargetDeviceServicesModule->GetDeviceProxyManager()->GetProxies(Platform->PlatformName(), false, DeviceProxies);
-
-			// if this platform had no devices, but we want to show an extra option if not installed right
-			if (DeviceProxies.Num() == 0)
-			{
-				FString NotInstalledDocLink;
-				// if the platform wasn't installed, we'll add a menu item later (we never care about code in this case, since we don't compile)
-				if (!Platform->IsSdkInstalled(false, NotInstalledDocLink) && PlatformsToMaybeInstallLinksFor.Find(PlatformInfo.DisplayName.ToString()) != INDEX_NONE)
-				{
-					PlatformsToAddInstallLinksFor.Add(Platform);
-				}
-			}
-			else
+			if (PlatformInfo.SDKStatus == PlatformInfo::EPlatformSDKStatus::Installed)
 			{
-				// for each proxy...
-				for (auto DeviceProxyIt = DeviceProxies.CreateIterator(); DeviceProxyIt; ++DeviceProxyIt)
-				{
-					ITargetDeviceProxyPtr DeviceProxy = *DeviceProxyIt;
-
-					// ... create an action...
-					FUIAction LaunchDeviceAction(
-						FExecuteAction::CreateStatic(&FInternalPlayWorldCommandCallbacks::HandleLaunchOnDeviceActionExecute, DeviceProxy->GetDeviceId(), DeviceProxy->GetName()),
-						FCanExecuteAction::CreateStatic(&FInternalPlayWorldCommandCallbacks::HandleLaunchOnDeviceActionCanExecute, DeviceProxy->GetDeviceId()),
-						FIsActionChecked::CreateStatic(&FInternalPlayWorldCommandCallbacks::HandleLaunchOnDeviceActionIsChecked, DeviceProxy->GetDeviceId())
-					);
+				ITargetPlatform* Platform = GetTargetPlatformManager()->FindTargetPlatform(PlatformInfo.VanillaPlatformName.ToString());
 
-					// ... generate display label...
-					FFormatNamedArguments LabelArguments;
-					LabelArguments.Add(TEXT("DeviceName"), FText::FromString(DeviceProxy->GetName()));
+				// for each platform...
+				TArray<ITargetDeviceProxyPtr> DeviceProxies;
+				TargetDeviceServicesModule->GetDeviceProxyManager()->GetProxies(Platform->PlatformName(), false, DeviceProxies);
 
-					if (!DeviceProxy->IsConnected())
-					{
-						LabelArguments.Add(TEXT("HostUser"), LOCTEXT("DisconnectedHint", " [Disconnected]"));
-					}
-					else if (DeviceProxy->GetHostUser() != FPlatformProcess::UserName(false))
-					{
-						LabelArguments.Add(TEXT("HostUser"), FText::FromString(DeviceProxy->GetHostUser()));
-					}
-					else
+				// if this platform had no devices, but we want to show an extra option if not installed right
+				if (DeviceProxies.Num() == 0)
+				{
+					if (PlatformsToMaybeInstallLinksFor.Find(PlatformInfo.VanillaPlatformName.ToString()) != INDEX_NONE && PlatformsWithNoDevices.Find(PlatformInfo.VanillaPlatformName.ToString()) == INDEX_NONE)
 					{
-						LabelArguments.Add(TEXT("HostUser"), FText::GetEmpty());
+						// add an entry with a no devices found
+						PlatformsWithNoDevices.Add(PlatformInfo.VanillaPlatformName.ToString());
+
+						// ... generate display label...
+						FFormatNamedArguments LabelArguments;
+						LabelArguments.Add(TEXT("DisplayName"), PlatformInfo.DisplayName);
+
+						FText Label = FText::Format(LOCTEXT("NoDeviceLabel", "{DisplayName} - No Devices Found"), LabelArguments);
+
+						// ... create an action...
+						FUIAction NoDeviceAction(
+							FExecuteAction::CreateStatic(&FInternalPlayWorldCommandCallbacks::HandleNoDeviceFoundActionExecute),
+							FCanExecuteAction::CreateStatic(&FInternalPlayWorldCommandCallbacks::HandleNoDeviceFoundActionCanExecute)
+							);
+
+						// ... generate tooltip text
+						FFormatNamedArguments TooltipArguments;
+						TooltipArguments.Add(TEXT("DisplayName"), PlatformInfo.DisplayName);
+						FText Tooltip = FText::Format(LOCTEXT("LaunchDeviceToolTipText", "Found no connected devices for {DisplayName}"), TooltipArguments);
+
+						// ... and add a menu entry
+						MenuBuilder.AddMenuEntry(
+							NoDeviceAction, 
+							ProjectTargetPlatformEditorModule.MakePlatformMenuItemWidget(PlatformInfo, true, Label), 
+							NAME_None, 
+							Tooltip, 
+							EUserInterfaceActionType::Check
+							);
 					}
-
-					FText Label = FText::Format(LOCTEXT("LaunchDeviceLabel", "{DeviceName}{HostUser}"), LabelArguments);
-
-					// ... generate tooltip text
-					FFormatNamedArguments TooltipArguments;
-					TooltipArguments.Add(TEXT("DeviceID"), FText::FromString(DeviceProxy->GetName()));
-					TooltipArguments.Add(TEXT("DisplayName"), PlatformInfo.DisplayName);
-					FText Tooltip = FText::Format(LOCTEXT("LaunchDeviceToolTipText", "Launch the game on this {DisplayName} device ({DeviceID})"), TooltipArguments);
-
-					FProjectStatus ProjectStatus;
-					if (IProjectManager::Get().QueryStatusForCurrentProject(ProjectStatus) && !ProjectStatus.IsTargetPlatformSupported(PlatformInfo.VanillaPlatformName)) 
+				}
+				else
+				{
+					// for each proxy...
+					for (auto DeviceProxyIt = DeviceProxies.CreateIterator(); DeviceProxyIt; ++DeviceProxyIt)
 					{
-						FText TooltipLine2 = FText::Format(LOCTEXT("LaunchDevicePlatformWarning", "{DisplayName} is not listed as a target platform for this project, so may not run as expected."), TooltipArguments);
-						Tooltip = FText::Format(FText::FromString(TEXT("{0}\n\n{1}")), Tooltip, TooltipLine2);
+						ITargetDeviceProxyPtr DeviceProxy = *DeviceProxyIt;
+
+						// ... create an action...
+						FUIAction LaunchDeviceAction(
+							FExecuteAction::CreateStatic(&FInternalPlayWorldCommandCallbacks::HandleLaunchOnDeviceActionExecute, DeviceProxy->GetDeviceId(), DeviceProxy->GetName()),
+							FCanExecuteAction::CreateStatic(&FInternalPlayWorldCommandCallbacks::HandleLaunchOnDeviceActionCanExecute, DeviceProxy->GetDeviceId()),
+							FIsActionChecked::CreateStatic(&FInternalPlayWorldCommandCallbacks::HandleLaunchOnDeviceActionIsChecked, DeviceProxy->GetDeviceId())
+							);
+
+						// ... generate display label...
+						FFormatNamedArguments LabelArguments;
+						LabelArguments.Add(TEXT("DeviceName"), FText::FromString(DeviceProxy->GetName()));
+
+						if (!DeviceProxy->IsConnected())
+						{
+							LabelArguments.Add(TEXT("HostUser"), LOCTEXT("DisconnectedHint", " [Disconnected]"));
+						}
+						else if (DeviceProxy->GetHostUser() != FPlatformProcess::UserName(false))
+						{
+							LabelArguments.Add(TEXT("HostUser"), FText::FromString(DeviceProxy->GetHostUser()));
+						}
+						else
+						{
+							LabelArguments.Add(TEXT("HostUser"), FText::GetEmpty());
+						}
+
+						FText Label = FText::Format(LOCTEXT("LaunchDeviceLabel", "{DeviceName}{HostUser}"), LabelArguments);
+
+						// ... generate tooltip text
+						FFormatNamedArguments TooltipArguments;
+						TooltipArguments.Add(TEXT("DeviceID"), FText::FromString(DeviceProxy->GetName()));
+						TooltipArguments.Add(TEXT("DisplayName"), PlatformInfo.DisplayName);
+						FText Tooltip = FText::Format(LOCTEXT("LaunchDeviceToolTipText", "Launch the game on this {DisplayName} device ({DeviceID})"), TooltipArguments);
+
+						FProjectStatus ProjectStatus;
+						if (IProjectManager::Get().QueryStatusForCurrentProject(ProjectStatus) && !ProjectStatus.IsTargetPlatformSupported(PlatformInfo.VanillaPlatformName)) 
+						{
+							FText TooltipLine2 = FText::Format(LOCTEXT("LaunchDevicePlatformWarning", "{DisplayName} is not listed as a target platform for this project, so may not run as expected."), TooltipArguments);
+							Tooltip = FText::Format(FText::FromString(TEXT("{0}\n\n{1}")), Tooltip, TooltipLine2);
+						}
+
+						// ... and add a menu entry
+						MenuBuilder.AddMenuEntry(
+							LaunchDeviceAction, 
+							ProjectTargetPlatformEditorModule.MakePlatformMenuItemWidget(PlatformInfo, true, Label), 
+							NAME_None, 
+							Tooltip, 
+							EUserInterfaceActionType::Check
+							);
 					}
-
-					// ... and add a menu entry
-					MenuBuilder.AddMenuEntry(
-						LaunchDeviceAction, 
-						ProjectTargetPlatformEditorModule.MakePlatformMenuItemWidget(PlatformInfo, true, Label), 
-						NAME_None, 
-						Tooltip, 
-						EUserInterfaceActionType::Check
-						);
+				}
+			}
+			else
+			{
+				// if the platform wasn't installed, we'll add a menu item later (we never care about code in this case, since we don't compile)
+				if (PlatformsToMaybeInstallLinksFor.Find(PlatformInfo.VanillaPlatformName.ToString()) != INDEX_NONE)
+				{
+					PlatformsToAddInstallLinksFor.Add(PlatformInfo);
 				}
 			}
 		}
@@ -679,12 +722,12 @@ TSharedRef< SWidget > FPlayWorldCommands::GenerateLaunchMenuContent( TSharedRef<
 		{
 			for (int32 PlatformIndex = 0; PlatformIndex < PlatformsToAddInstallLinksFor.Num(); PlatformIndex++)
 			{
-				ITargetPlatform* Platform = PlatformsToAddInstallLinksFor[PlatformIndex];
+				const PlatformInfo::FPlatformInfo& Platform = PlatformsToAddInstallLinksFor[PlatformIndex];
 
-				FUIAction Action(FExecuteAction::CreateStatic(&FInternalPlayWorldCommandCallbacks::HandleShowSDKTutorial, Platform));
+				FUIAction Action(FExecuteAction::CreateStatic(&FInternalPlayWorldCommandCallbacks::HandleShowSDKTutorial, Platform.DisplayName.ToString(), Platform.SDKTutorial));
 
 				FFormatNamedArguments LabelArguments;
-				LabelArguments.Add(TEXT("PlatformName"), Platform->DisplayName());
+				LabelArguments.Add(TEXT("PlatformName"), Platform.DisplayName);
 				FText Label = FText::Format(LOCTEXT("LaunchPlatformLabel", "{PlatformName} Support"), LabelArguments);
 
 
@@ -1336,6 +1379,15 @@ bool FInternalPlayWorldCommandCallbacks::IsReadyToLaunchOnDevice(FString DeviceI
 			break;
 		}
 	}
+	else
+	{
+		const PlatformInfo::FPlatformInfo* PlatformInfo = PlatformInfo::FindPlatformInfo(FName(*PlatformName));
+		check(PlatformInfo);
+
+		IMainFrameModule& MainFrameModule = FModuleManager::GetModuleChecked<IMainFrameModule>(TEXT("MainFrame"));
+		MainFrameModule.BroadcastMainFrameSDKNotInstalled(PlatformInfo->TargetPlatformName.ToString(), PlatformInfo->SDKTutorial);
+		return false;
+	}
 
 	return true;
 }
@@ -1369,15 +1421,11 @@ bool FInternalPlayWorldCommandCallbacks::HandleLaunchOnDeviceActionIsChecked(FSt
 }
 
 
-void FInternalPlayWorldCommandCallbacks::HandleShowSDKTutorial( ITargetPlatform* Platform )
+void FInternalPlayWorldCommandCallbacks::HandleShowSDKTutorial( FString PlatformName, FString NotInstalledDocLink )
 {
-	// we know it's not installed, so just get the link
-	FString NotInstalledDocLink;
-	Platform->IsSdkInstalled(false, NotInstalledDocLink);
-
 	// broadcast this, and assume someone will pick it up
 	IMainFrameModule& MainFrameModule = FModuleManager::GetModuleChecked<IMainFrameModule>(TEXT("MainFrame"));
-	MainFrameModule.BroadcastMainFrameSDKNotInstalled(Platform->DisplayName().ToString(), NotInstalledDocLink);
+	MainFrameModule.BroadcastMainFrameSDKNotInstalled(PlatformName, NotInstalledDocLink);
 }
 
 FSlateIcon FInternalPlayWorldCommandCallbacks::GetResumePlaySessionImage()
diff --git a/Engine/Source/Editor/UnrealEd/UnrealEd.Build.cs b/Engine/Source/Editor/UnrealEd/UnrealEd.Build.cs
--- a/Engine/Source/Editor/UnrealEd/UnrealEd.Build.cs
+++ b/Engine/Source/Editor/UnrealEd/UnrealEd.Build.cs
@@ -77,6 +77,7 @@ public class UnrealEd : ModuleRules
 				"BlueprintGraph",
 				"OnlineSubsystem",
 				"OnlineBlueprintSupport",
+				"DesktopPlatform",
 				"DirectoryWatcher",
 				"EditorSettingsViewer",
 				"EditorStyle",
