Add an additional check to determine valid engine directories, so running UnrealVersionSelector directly from the launcher installation folder won't register one there.

diff --git a/Engine/Source/Developer/DesktopPlatform/Private/DesktopPlatformBase.cpp b/Engine/Source/Developer/DesktopPlatform/Private/DesktopPlatformBase.cpp
--- a/Engine/Source/Developer/DesktopPlatform/Private/DesktopPlatformBase.cpp
+++ b/Engine/Source/Developer/DesktopPlatform/Private/DesktopPlatformBase.cpp
@@ -209,9 +209,24 @@ bool FDesktopPlatformBase::IsPerforceBuild(const FString &EngineRootDir)
 
 bool FDesktopPlatformBase::IsValidRootDirectory(const FString &RootDir)
 {
+	// Check that there's an Engine\Binaries directory underneath the root
 	FString EngineBinariesDirName = RootDir / TEXT("Engine/Binaries");
 	FPaths::NormalizeDirectoryName(EngineBinariesDirName);
-	return IFileManager::Get().DirectoryExists(*EngineBinariesDirName);
+	if(!IFileManager::Get().DirectoryExists(*EngineBinariesDirName))
+	{
+		return false;
+	}
+
+	// Also check there's an Engine\Build directory. This will filter out anything that has an engine-like directory structure but doesn't allow building code projects - like the launcher.
+	FString EngineBuildDirName = RootDir / TEXT("Engine/Build");
+	FPaths::NormalizeDirectoryName(EngineBuildDirName);
+	if(!IFileManager::Get().DirectoryExists(*EngineBuildDirName))
+	{
+		return false;
+	}
+
+	// Otherwise it's valid
+	return true;
 }
 
 bool FDesktopPlatformBase::SetEngineIdentifierForProject(const FString &ProjectFileName, const FString &InIdentifier)
diff --git a/Engine/Source/Programs/UnrealVersionSelector/Private/UnrealVersionSelector.cpp b/Engine/Source/Programs/UnrealVersionSelector/Private/UnrealVersionSelector.cpp
--- a/Engine/Source/Programs/UnrealVersionSelector/Private/UnrealVersionSelector.cpp
+++ b/Engine/Source/Programs/UnrealVersionSelector/Private/UnrealVersionSelector.cpp
@@ -10,14 +10,18 @@ bool GenerateProjectFiles(const FString& ProjectFileName);
 bool RegisterCurrentEngineDirectory()
 {
 	// Prompt for registering this directory
-	if(FPlatformMisc::MessageBoxExt(EAppMsgType::YesNo, TEXT("Configure this directory as an Unreal Engine installation?"), TEXT("Question")) != EAppReturnType::Yes)
+	if(FPlatformMisc::MessageBoxExt(EAppMsgType::YesNo, TEXT("Register this directory as an Unreal Engine installation?"), TEXT("Question")) != EAppReturnType::Yes)
 	{
 		return false;
 	}
 
 	// Get the current engine directory.
 	FString EngineRootDir = FPlatformProcess::BaseDir();
-	FPlatformInstallation::NormalizeEngineRootDir(EngineRootDir);
+	if(!FPlatformInstallation::NormalizeEngineRootDir(EngineRootDir))
+	{
+		FPlatformMisc::MessageBoxExt(EAppMsgType::Ok, TEXT("The current folder does not contain an engine installation."), TEXT("Error"));
+		return false;
+	}
 
 	// Get any existing tag name or register a new one
 	FString Identifier;
