[INTEGRATE] UE4-LauncherRelease->UE4 integrate integrate latest from LauncherReleases, including fixes for various issues that we've released the past month, and the new engine slot and notification systems

diff --git a/Engine/Source/Developer/DesktopPlatform/Private/Mac/DesktopPlatformMac.cpp b/Engine/Source/Developer/DesktopPlatform/Private/Mac/DesktopPlatformMac.cpp
--- a/Engine/Source/Developer/DesktopPlatform/Private/Mac/DesktopPlatformMac.cpp
+++ b/Engine/Source/Developer/DesktopPlatform/Private/Mac/DesktopPlatformMac.cpp
@@ -567,6 +567,20 @@ void FDesktopPlatformMac::EnumerateEngineInstallations(TMap<FString, FString> &O
 	ConfigFile.Read(ConfigPath);
 
 	FConfigSection &Section = ConfigFile.FindOrAdd(TEXT("Installations"));
+	// Remove invalid entries
+	TArray<FName> KeysToRemove;
+	for (auto It : Section)
+	{
+		const FString& EngineDir = It.Value;
+		if (EngineDir.Contains("Unreal Engine.app/Contents/") || EngineDir.Contains("[REDACTED_USER_HOME]/UnrealEngine/Launcher") || !IFileManager::Get().DirectoryExists(*EngineDir))
+		{
+			KeysToRemove.Add(It.Key);
+		}
+	}
+	for (auto Key : KeysToRemove)
+	{
+		Section.Remove(Key);
+	}
 
 	CFArrayRef AllApps = LSCopyApplicationURLsForURL((__bridge CFURLRef)[NSURL fileURLWithPath:UProjectPath.GetNSString()], kLSRolesAll);
 	if (AllApps)
@@ -578,7 +592,7 @@ void FDesktopPlatformMac::EnumerateEngineInstallations(TMap<FString, FString> &O
 			NSBundle* AppBundle = [NSBundle bundleWithURL:AppURL];
 			FString EngineDir = FString([[AppBundle bundlePath] stringByDeletingLastPathComponent]);
 			if (([[AppBundle bundleIdentifier] isEqualToString:@"com.epicgames.UE4Editor"] || [[AppBundle bundleIdentifier] isEqualToString:@"com.epicgames.UE4EditorServices"])
-				&& EngineDir.RemoveFromEnd(TEXT("/Engine/Binaries/Mac")))
+				&& EngineDir.RemoveFromEnd(TEXT("/Engine/Binaries/Mac")) && !EngineDir.Contains("Unreal Engine.app/Contents/") && !EngineDir.Contains("[REDACTED_USER_HOME]/UnrealEngine/Launcher"))
 			{
 				FString EngineId;
 				const FName* Key = Section.FindKey(EngineDir);
diff --git a/Engine/Source/Programs/Mac/UE4EditorServices/Private/Mac/UE4EditorServicesAppDelegate.cpp b/Engine/Source/Programs/Mac/UE4EditorServices/Private/Mac/UE4EditorServicesAppDelegate.cpp
--- a/Engine/Source/Programs/Mac/UE4EditorServices/Private/Mac/UE4EditorServicesAppDelegate.cpp
+++ b/Engine/Source/Programs/Mac/UE4EditorServices/Private/Mac/UE4EditorServicesAppDelegate.cpp
@@ -404,7 +404,8 @@
 - (BOOL)isAppValidForUProjectFiles:(NSBundle*)AppBundle
 {
 	return AppBundle && ([[AppBundle bundleIdentifier] isEqualToString:@"com.epicgames.UE4Editor"] || [[AppBundle bundleIdentifier] isEqualToString:@"com.epicgames.UE4EditorServices"])
-					 && [[[AppBundle bundlePath] stringByDeletingLastPathComponent] hasSuffix:@"Engine/Binaries/Mac"];
+					 && [[[AppBundle bundlePath] stringByDeletingLastPathComponent] hasSuffix:@"Engine/Binaries/Mac"]
+					 && [[AppBundle bundlePath] rangeOfString:@".app/Contents/UE4"].location == NSNotFound;
 }
 
 - (NSString*)enginePathForAppBundle:(NSBundle*)AppBundle
