Fix for Windows XP compile issues.

Cannot use ::RegGetValue or ::RegDeleteTree with xp supported builds, compatible alternatives added.

diff --git a/Engine/Source/Developer/DesktopPlatform/Private/Windows/DesktopPlatformWindows.cpp b/Engine/Source/Developer/DesktopPlatform/Private/Windows/DesktopPlatformWindows.cpp
--- a/Engine/Source/Developer/DesktopPlatform/Private/Windows/DesktopPlatformWindows.cpp
+++ b/Engine/Source/Developer/DesktopPlatform/Private/Windows/DesktopPlatformWindows.cpp
@@ -529,9 +529,8 @@ void FDesktopPlatformWindows::GetRequiredRegistrySettings(TIndirectArray<FRegist
 	FString ExecutableFileName = FString(FPlatformProcess::BaseDir()) / DefaultVersionSelectorName;
 
 	// Defer to UnrealVersionSelector.exe in a launcher installation if it's got the same version number or greater.
-	TCHAR InstallDir[MAX_PATH];
-	::DWORD InstallDirSize = sizeof(InstallDir);
-	if(RegGetValue(HKEY_LOCAL_MACHINE, TEXT("SOFTWARE\\EpicGames\\Unreal Engine"), TEXT("INSTALLDIR"), RRF_RT_REG_SZ, NULL, InstallDir, &InstallDirSize) == ERROR_SUCCESS)
+	FString InstallDir;
+	if (FWindowsPlatformMisc::QueryRegKey(HKEY_LOCAL_MACHINE, TEXT("SOFTWARE\\EpicGames\\Unreal Engine"), TEXT("INSTALLDIR"), InstallDir) && (InstallDir.Len() > 0))
 	{
 		FString NormalizedInstallDir = InstallDir;
 		FPaths::NormalizeDirectoryName(NormalizedInstallDir);
diff --git a/Engine/Source/Developer/DesktopPlatform/Private/Windows/WindowsRegistry.cpp b/Engine/Source/Developer/DesktopPlatform/Private/Windows/WindowsRegistry.cpp
--- a/Engine/Source/Developer/DesktopPlatform/Private/Windows/WindowsRegistry.cpp
+++ b/Engine/Source/Developer/DesktopPlatform/Private/Windows/WindowsRegistry.cpp
@@ -4,6 +4,11 @@
 #include "WindowsRegistry.h"
 #include "AllowWindowsPlatformTypes.h"
 
+#if WINVER == 0x0502
+	// Some registry functions are not available on windows XP, so we need the XP compat alternative
+	#include <Shlwapi.h>
+#endif
+
 ////////////////////////////////////////////////////////////////////
 
 FRegistryValue::FRegistryValue()
@@ -178,7 +183,12 @@ bool FRegistryKey::Write(HKEY hKey, bool bRemoveDifferences) const
 		// Remove any keys that don't exist any more
 		for (int32 Idx = 0; Idx < KeyNames.Num(); Idx++)
 		{
+#if WINVER == 0x0502
+			if (!Keys.Contains(KeyNames[Idx]) && SHDeleteKey(hKey, *KeyNames[Idx]) != ERROR_SUCCESS)
+#else
+
 			if (!Keys.Contains(KeyNames[Idx]) && RegDeleteTree(hKey, *KeyNames[Idx]) != ERROR_SUCCESS)
+#endif // WINVER == 0x0502
 			{
 				return false;
 			}
