New ENUM_CLASS_FLAGS macro for declaring all of the bitwise operators needed for a bitflags enum.

Fixed CONSTEXPR definitions.
Converted EFontImportFlags to an enum class in order to test this feature.

diff --git a/Engine/Source/Developer/DesktopPlatform/Private/Linux/DesktopPlatformLinux.cpp b/Engine/Source/Developer/DesktopPlatform/Private/Linux/DesktopPlatformLinux.cpp
--- a/Engine/Source/Developer/DesktopPlatform/Private/Linux/DesktopPlatformLinux.cpp
+++ b/Engine/Source/Developer/DesktopPlatform/Private/Linux/DesktopPlatformLinux.cpp
@@ -81,7 +81,7 @@ bool FDesktopPlatformLinux::OpenDirectoryDialog(const void* ParentWindowHandle,
 #endif // WITH_LINUX_NATIVE_DIALOGS
 }
 
-bool FDesktopPlatformLinux::OpenFontDialog(const void* ParentWindowHandle, FString& OutFontName, float& OutHeight, EFontImportFlags::Type& OutFlags)
+bool FDesktopPlatformLinux::OpenFontDialog(const void* ParentWindowHandle, FString& OutFontName, float& OutHeight, EFontImportFlags& OutFlags)
 {
 	unimplemented();
 	return false;
diff --git a/Engine/Source/Developer/DesktopPlatform/Private/Linux/DesktopPlatformLinux.h b/Engine/Source/Developer/DesktopPlatform/Private/Linux/DesktopPlatformLinux.h
--- a/Engine/Source/Developer/DesktopPlatform/Private/Linux/DesktopPlatformLinux.h
+++ b/Engine/Source/Developer/DesktopPlatform/Private/Linux/DesktopPlatformLinux.h
@@ -12,7 +12,7 @@ public:
 	virtual bool OpenFileDialog(const void* ParentWindowHandle, const FString& DialogTitle, const FString& DefaultPath, const FString& DefaultFile, const FString& FileTypes, uint32 Flags, TArray<FString>& OutFilenames, int32& OutFilterIndex) override;
 	virtual bool SaveFileDialog(const void* ParentWindowHandle, const FString& DialogTitle, const FString& DefaultPath, const FString& DefaultFile, const FString& FileTypes, uint32 Flags, TArray<FString>& OutFilenames) override;
 	virtual bool OpenDirectoryDialog(const void* ParentWindowHandle, const FString& DialogTitle, const FString& DefaultPath, FString& OutFolderName) override;
-	virtual bool OpenFontDialog(const void* ParentWindowHandle, FString& OutFontName, float& OutHeight, EFontImportFlags::Type& OutFlags) override;
+	virtual bool OpenFontDialog(const void* ParentWindowHandle, FString& OutFontName, float& OutHeight, EFontImportFlags& OutFlags) override;
 	virtual bool OpenLauncher(bool Install, FString CommandLineParams ) override;
 
 	virtual bool RegisterEngineInstallation(const FString &RootDir, FString &OutIdentifier) override;
diff --git a/Engine/Source/Developer/DesktopPlatform/Private/Mac/DesktopPlatformMac.cpp b/Engine/Source/Developer/DesktopPlatform/Private/Mac/DesktopPlatformMac.cpp
--- a/Engine/Source/Developer/DesktopPlatform/Private/Mac/DesktopPlatformMac.cpp
+++ b/Engine/Source/Developer/DesktopPlatform/Private/Mac/DesktopPlatformMac.cpp
@@ -268,7 +268,7 @@ bool FDesktopPlatformMac::OpenDirectoryDialog(const void* ParentWindowHandle, co
 	return bSuccess;
 }
 
-bool FDesktopPlatformMac::OpenFontDialog(const void* ParentWindowHandle, FString& OutFontName, float& OutHeight, EFontImportFlags::Type& OutFlags)
+bool FDesktopPlatformMac::OpenFontDialog(const void* ParentWindowHandle, FString& OutFontName, float& OutHeight, EFontImportFlags& OutFlags)
 {
 	SCOPED_AUTORELEASE_POOL;
 	FCocoaScopeContext ContextGuard;
@@ -304,14 +304,14 @@ bool FDesktopPlatformMac::OpenFontDialog(const void* ParentWindowHandle, FString
 		OutFontName = FontName;
 		OutHeight = [Font pointSize];
 
-		uint32 FontFlags = EFontImportFlags::None;
+		auto FontFlags = EFontImportFlags::None;
 
 		if( [Font underlineThickness] >= 1.0 )
 		{
 			FontFlags |= EFontImportFlags::EnableUnderline;
 		}
 
-		OutFlags = static_cast<EFontImportFlags::Type>(FontFlags);
+		OutFlags = FontFlags;
 	}
 
 	MacApplication->ResetModifierKeys();
diff --git a/Engine/Source/Developer/DesktopPlatform/Private/Mac/DesktopPlatformMac.h b/Engine/Source/Developer/DesktopPlatform/Private/Mac/DesktopPlatformMac.h
--- a/Engine/Source/Developer/DesktopPlatform/Private/Mac/DesktopPlatformMac.h
+++ b/Engine/Source/Developer/DesktopPlatform/Private/Mac/DesktopPlatformMac.h
@@ -12,7 +12,7 @@ public:
 	virtual bool OpenFileDialog(const void* ParentWindowHandle, const FString& DialogTitle, const FString& DefaultPath, const FString& DefaultFile, const FString& FileTypes, uint32 Flags, TArray<FString>& OutFilenames, int32& OutFilterIndex) override;
 	virtual bool SaveFileDialog(const void* ParentWindowHandle, const FString& DialogTitle, const FString& DefaultPath, const FString& DefaultFile, const FString& FileTypes, uint32 Flags, TArray<FString>& OutFilenames) override;
 	virtual bool OpenDirectoryDialog(const void* ParentWindowHandle, const FString& DialogTitle, const FString& DefaultPath, FString& OutFolderName) override;
-	virtual bool OpenFontDialog(const void* ParentWindowHandle, FString& OutFontName, float& OutHeight, EFontImportFlags::Type& OutFlags) override;
+	virtual bool OpenFontDialog(const void* ParentWindowHandle, FString& OutFontName, float& OutHeight, EFontImportFlags& OutFlags) override;
 	virtual bool OpenLauncher(bool Install, FString CommandLineParams ) override;
 
 	virtual bool RegisterEngineInstallation(const FString &RootDir, FString &OutIdentifier) override;
diff --git a/Engine/Source/Developer/DesktopPlatform/Private/Windows/DesktopPlatformWindows.cpp b/Engine/Source/Developer/DesktopPlatform/Private/Windows/DesktopPlatformWindows.cpp
--- a/Engine/Source/Developer/DesktopPlatform/Private/Windows/DesktopPlatformWindows.cpp
+++ b/Engine/Source/Developer/DesktopPlatform/Private/Windows/DesktopPlatformWindows.cpp
@@ -93,7 +93,7 @@ bool FDesktopPlatformWindows::OpenDirectoryDialog(const void* ParentWindowHandle
 	return bSuccess;
 }
 
-bool FDesktopPlatformWindows::OpenFontDialog(const void* ParentWindowHandle, FString& OutFontName, float& OutHeight, EFontImportFlags::Type& OutFlags)
+bool FDesktopPlatformWindows::OpenFontDialog(const void* ParentWindowHandle, FString& OutFontName, float& OutHeight, EFontImportFlags& OutFlags)
 {	
 	FScopedSystemModalMode SystemModalScope;
 
@@ -113,7 +113,7 @@ bool FDesktopPlatformWindows::OpenFontDialog(const void* ParentWindowHandle, FSt
 		HDC DC = ::GetDC( cf.hwndOwner ); 
 		const float LogicalPixelsY = static_cast<float>(GetDeviceCaps(DC, LOGPIXELSY));
 		const int32 PixelHeight = static_cast<int32>(-lf.lfHeight * ( 72.0f / LogicalPixelsY ));	// Always target 72 DPI
-		uint32 FontFlags = EFontImportFlags::None;
+		auto FontFlags = EFontImportFlags::None;
 		if ( lf.lfUnderline )
 		{
 			FontFlags |= EFontImportFlags::EnableUnderline;
@@ -129,7 +129,7 @@ bool FDesktopPlatformWindows::OpenFontDialog(const void* ParentWindowHandle, FSt
 
 		OutFontName = (const TCHAR*)lf.lfFaceName;
 		OutHeight = PixelHeight;
-		OutFlags = static_cast<EFontImportFlags::Type>(FontFlags);
+		OutFlags = FontFlags;
 
 		::ReleaseDC( cf.hwndOwner, DC ); 
 	}
diff --git a/Engine/Source/Developer/DesktopPlatform/Private/Windows/DesktopPlatformWindows.h b/Engine/Source/Developer/DesktopPlatform/Private/Windows/DesktopPlatformWindows.h
--- a/Engine/Source/Developer/DesktopPlatform/Private/Windows/DesktopPlatformWindows.h
+++ b/Engine/Source/Developer/DesktopPlatform/Private/Windows/DesktopPlatformWindows.h
@@ -13,7 +13,7 @@ public:
 	virtual bool OpenFileDialog(const void* ParentWindowHandle, const FString& DialogTitle, const FString& DefaultPath, const FString& DefaultFile, const FString& FileTypes, uint32 Flags, TArray<FString>& OutFilenames, int32& OutFilterIndex) override;
 	virtual bool SaveFileDialog(const void* ParentWindowHandle, const FString& DialogTitle, const FString& DefaultPath, const FString& DefaultFile, const FString& FileTypes, uint32 Flags, TArray<FString>& OutFilenames) override;
 	virtual bool OpenDirectoryDialog(const void* ParentWindowHandle, const FString& DialogTitle, const FString& DefaultPath, FString& OutFolderName) override;
-	virtual bool OpenFontDialog(const void* ParentWindowHandle, FString& OutFontName, float& OutHeight, EFontImportFlags::Type& OutFlags) override;
+	virtual bool OpenFontDialog(const void* ParentWindowHandle, FString& OutFontName, float& OutHeight, EFontImportFlags& OutFlags) override;
 	virtual bool OpenLauncher(bool Install, FString CommandLineParams ) override;
 
 	virtual bool RegisterEngineInstallation(const FString &RootDir, FString &OutIdentifier) override;
diff --git a/Engine/Source/Developer/DesktopPlatform/Public/IDesktopPlatform.h b/Engine/Source/Developer/DesktopPlatform/Public/IDesktopPlatform.h
--- a/Engine/Source/Developer/DesktopPlatform/Public/IDesktopPlatform.h
+++ b/Engine/Source/Developer/DesktopPlatform/Public/IDesktopPlatform.h
@@ -11,24 +11,23 @@ namespace EFileDialogFlags
 	};
 }
 
-namespace EFontImportFlags
+enum class EFontImportFlags
 {
-	enum Type
-	{
-		None = 0x0,						// No flags
-		EnableAntialiasing = 0x1,		// Whether the font should be antialiased or not.  Usually you should leave this enabled.
-		EnableBold = 0x2,				// Whether the font should be generated in bold or not
-		EnableItalic = 0x4,				// Whether the font should be generated in italics or not
-		EnableUnderline = 0x8,			// Whether the font should be generated with an underline or not
-		AlphaOnly = 0x10,				// Forces PF_G8 and only maintains Alpha value and discards color
-		CreatePrintableOnly = 0x20,		// Skips generation of glyphs for any characters that are not considered 'printable'
-		IncludeASCIIRange = 0x40,		// When specifying a range of characters and this is enabled, forces ASCII characters (0 thru 255) to be included as well
-		EnableDropShadow = 0x80,		// Enables a very simple, 1-pixel, black colored drop shadow for the generated font
-		EnableLegacyMode = 0x100,		// Enables legacy font import mode.  This results in lower quality antialiasing and larger glyph bounds, but may be useful when debugging problems
-		UseDistanceFieldAlpha = 0x200	// Alpha channel of the font textures will store a distance field instead of a color mask
-	};
+	None = 0x0,						// No flags
+	EnableAntialiasing = 0x1,		// Whether the font should be antialiased or not.  Usually you should leave this enabled.
+	EnableBold = 0x2,				// Whether the font should be generated in bold or not
+	EnableItalic = 0x4,				// Whether the font should be generated in italics or not
+	EnableUnderline = 0x8,			// Whether the font should be generated with an underline or not
+	AlphaOnly = 0x10,				// Forces PF_G8 and only maintains Alpha value and discards color
+	CreatePrintableOnly = 0x20,		// Skips generation of glyphs for any characters that are not considered 'printable'
+	IncludeASCIIRange = 0x40,		// When specifying a range of characters and this is enabled, forces ASCII characters (0 thru 255) to be included as well
+	EnableDropShadow = 0x80,		// Enables a very simple, 1-pixel, black colored drop shadow for the generated font
+	EnableLegacyMode = 0x100,		// Enables legacy font import mode.  This results in lower quality antialiasing and larger glyph bounds, but may be useful when debugging problems
+	UseDistanceFieldAlpha = 0x200	// Alpha channel of the font textures will store a distance field instead of a color mask
 };
 
+ENUM_CLASS_FLAGS(EFontImportFlags)
+
 
 /**
  * When constructed leaves system wide modal mode (all windows disabled except for the OS modal window)
@@ -122,7 +121,7 @@ public:
 	 * @param OutFlags					Any special flags the font has been tagged with
 	 * @return true if font choice was successfully selected
 	 */
-	virtual bool OpenFontDialog(const void* ParentWindowHandle, FString& OutFontName, float& OutHeight, EFontImportFlags::Type& OutFlags) = 0;
+	virtual bool OpenFontDialog(const void* ParentWindowHandle, FString& OutFontName, float& OutHeight, EFontImportFlags& OutFlags) = 0;
 
 	/**
 	 * Opens the marketplace user interface.
diff --git a/Engine/Source/Developer/Profiler/Private/ProfilerDataSource.h b/Engine/Source/Developer/Profiler/Private/ProfilerDataSource.h
--- a/Engine/Source/Developer/Profiler/Private/ProfilerDataSource.h
+++ b/Engine/Source/Developer/Profiler/Private/ProfilerDataSource.h
@@ -808,7 +808,7 @@ namespace EEventPropertyTypes
 			Type( GetTypeFromFormatter(PropertyFormatter) )
 		{}
 
-		CONSTEXPR EEventPropertyTypes::Type GetTypeFromFormatter( const EEventPropertyFormatters::Type PropertyFormatter ) const
+		EEventPropertyTypes::Type GetTypeFromFormatter( const EEventPropertyFormatters::Type PropertyFormatter ) const
 		{
 			switch( PropertyFormatter )
 			{
diff --git a/Engine/Source/Editor/UnrealEd/Private/Factories/TTFontImport.cpp b/Engine/Source/Editor/UnrealEd/Private/Factories/TTFontImport.cpp
--- a/Engine/Source/Editor/UnrealEd/Private/Factories/TTFontImport.cpp
+++ b/Engine/Source/Editor/UnrealEd/Private/Factories/TTFontImport.cpp
@@ -98,7 +98,7 @@ bool UTrueTypeFontFactory::ConfigureProperties()
 			SetupFontImportOptions();
 		}
 
-		EFontImportFlags::Type FontFlags;
+		EFontImportFlags FontFlags;
 		bFontSelected = DesktopPlatform->OpenFontDialog(
 			ParentWindowWindowHandle,
 			ImportOptions->Data.FontName,
@@ -108,15 +108,15 @@ bool UTrueTypeFontFactory::ConfigureProperties()
 
 		if ( bFontSelected )
 		{
-			if( FontFlags & EFontImportFlags::EnableUnderline )
+			if( !!(FontFlags & EFontImportFlags::EnableUnderline) )
 			{
 				ImportOptions->Data.bEnableUnderline = true;
 			}
-			if( FontFlags & EFontImportFlags::EnableItalic )
+			if( !!(FontFlags & EFontImportFlags::EnableItalic) )
 			{
 				ImportOptions->Data.bEnableItalic = true;
 			}
-			if( FontFlags & EFontImportFlags::EnableBold )
+			if( !!(FontFlags & EFontImportFlags::EnableBold) )
 			{
 				ImportOptions->Data.bEnableBold = true;
 			}
diff --git a/Engine/Source/Runtime/Core/Public/Core.h b/Engine/Source/Runtime/Core/Public/Core.h
--- a/Engine/Source/Runtime/Core/Public/Core.h
+++ b/Engine/Source/Runtime/Core/Public/Core.h
@@ -472,6 +472,7 @@ extern CORE_API TArray<struct FScriptTraceStackNode> GScriptStack;
 #include "MonitoredProcess.h"
 #include "Attribute.h"
 #include "Optional.h"
+#include "EnumClassFlags.h"
 
 #ifdef TRUE
 	#undef TRUE
diff --git a/Engine/Source/Runtime/Core/Public/HAL/Platform.h b/Engine/Source/Runtime/Core/Public/HAL/Platform.h
--- a/Engine/Source/Runtime/Core/Public/HAL/Platform.h
+++ b/Engine/Source/Runtime/Core/Public/HAL/Platform.h
@@ -324,7 +324,7 @@
 	#define ABSTRACT
 #endif
 #ifndef CONSTEXPR
-	#define CONSTEXPR
+	#define CONSTEXPR constexpr
 #endif
 
 // String constants
diff --git a/Engine/Source/Runtime/Core/Public/Misc/EnumClassFlags.h b/Engine/Source/Runtime/Core/Public/Misc/EnumClassFlags.h
new file mode 100644
--- /dev/null
+++ b/Engine/Source/Runtime/Core/Public/Misc/EnumClassFlags.h
@@ -0,0 +1,14 @@
+// Copyright 1998-2014 Epic Games, Inc. All Rights Reserved.
+
+#pragma once
+
+// Defines all bitwise operators for enum classes so it can be (mostly) used as a regular flags enum
+#define ENUM_CLASS_FLAGS(Enum) \
+	inline           Enum& operator|=(Enum& Lhs, Enum Rhs) { return Lhs = (Enum)((__underlying_type(Enum))Lhs | (__underlying_type(Enum))Rhs); } \
+	inline           Enum& operator&=(Enum& Lhs, Enum Rhs) { return Lhs = (Enum)((__underlying_type(Enum))Lhs & (__underlying_type(Enum))Rhs); } \
+	inline           Enum& operator^=(Enum& Lhs, Enum Rhs) { return Lhs = (Enum)((__underlying_type(Enum))Lhs ^ (__underlying_type(Enum))Rhs); } \
+	inline CONSTEXPR Enum  operator| (Enum  Lhs, Enum Rhs) { return (Enum)((__underlying_type(Enum))Lhs | (__underlying_type(Enum))Rhs); } \
+	inline CONSTEXPR Enum  operator& (Enum  Lhs, Enum Rhs) { return (Enum)((__underlying_type(Enum))Lhs & (__underlying_type(Enum))Rhs); } \
+	inline CONSTEXPR Enum  operator^ (Enum  Lhs, Enum Rhs) { return (Enum)((__underlying_type(Enum))Lhs ^ (__underlying_type(Enum))Rhs); } \
+	inline CONSTEXPR bool  operator! (Enum  E)             { return !(__underlying_type(Enum))E; } \
+	inline CONSTEXPR Enum  operator~ (Enum  E)             { return (Enum)~(__underlying_type(Enum))E; }
diff --git a/Engine/Source/Runtime/Core/Public/WinRT/WinRTPlatform.h b/Engine/Source/Runtime/Core/Public/WinRT/WinRTPlatform.h
--- a/Engine/Source/Runtime/Core/Public/WinRT/WinRTPlatform.h
+++ b/Engine/Source/Runtime/Core/Public/WinRT/WinRTPlatform.h
@@ -79,6 +79,7 @@ typedef FWinRTTypes FPlatformTypes;
 
 #pragma warning(disable : 4481) // nonstandard extension used: override specifier 'override'
 #define ABSTRACT abstract
+#define CONSTEXPR
 
 // Strings.
 #define LINE_TERMINATOR		TEXT("\r\n")
diff --git a/Engine/Source/Runtime/Core/Public/Windows/WIndowsPlatform.h b/Engine/Source/Runtime/Core/Public/Windows/WIndowsPlatform.h
--- a/Engine/Source/Runtime/Core/Public/Windows/WIndowsPlatform.h
+++ b/Engine/Source/Runtime/Core/Public/Windows/WIndowsPlatform.h
@@ -113,6 +113,7 @@ typedef FWindowsPlatformTypes FPlatformTypes;
 	#define CONSTEXPR constexpr
 #endif
 #define ABSTRACT abstract
+#define CONSTEXPR
 
 // Strings.
 #define LINE_TERMINATOR TEXT("\r\n")
