SlateViewer now compiles and runs using Clang for Windows
- Fixed bad casting in "slow task" window creation code
- Also added comment about treating warnings as errors with Clang

diff --git a/Engine/Source/Developer/DesktopPlatform/Private/Windows/WindowsNativeFeedbackContext.cpp b/Engine/Source/Developer/DesktopPlatform/Private/Windows/WindowsNativeFeedbackContext.cpp
--- a/Engine/Source/Developer/DesktopPlatform/Private/Windows/WindowsNativeFeedbackContext.cpp
+++ b/Engine/Source/Developer/DesktopPlatform/Private/Windows/WindowsNativeFeedbackContext.cpp
@@ -257,7 +257,8 @@ DWORD FWindowsNativeFeedbackContext::SlowTaskThreadProc(void* ThreadParam)
 	WindowRect.bottom = WindowRect.top + Params.StandardH;
 	AdjustWindowRectEx(&WindowRect, WindowStyle, 0, 0);
 
-	HWND hWnd = CreateWindow((TCHAR*)WndClassAtom, TEXT("Unreal Engine"), WindowStyle, WindowRect.left, WindowRect.top, WindowRect.right - WindowRect.left, WindowRect.bottom - WindowRect.top, NULL, NULL, hInstance, NULL);
+	const TCHAR* WindowClassName = MAKEINTATOM( WndClassAtom );
+	HWND hWnd = CreateWindow(WindowClassName, TEXT("Unreal Engine"), WindowStyle, WindowRect.left, WindowRect.top, WindowRect.right - WindowRect.left, WindowRect.bottom - WindowRect.top, NULL, NULL, hInstance, NULL);
 	SetWindowLongPtr(hWnd, GWLP_USERDATA, (LONG_PTR)&Params);
 	SendMessage(hWnd, WM_SETFONT, (WPARAM)hFont, 0);
 
@@ -325,7 +326,7 @@ DWORD FWindowsNativeFeedbackContext::SlowTaskThreadProc(void* ThreadParam)
 	DestroyWindow(hWnd);
 	DeleteObject(hLogFont);
 	DeleteObject(hFont);
-	UnregisterClass((TCHAR*)WndClassAtom, hInstance);
+	UnregisterClass(WindowClassName, hInstance);
 
 	return 0;
 }
diff --git a/Engine/Source/Programs/UnrealBuildTool/Windows/VCToolChain.cs b/Engine/Source/Programs/UnrealBuildTool/Windows/VCToolChain.cs
--- a/Engine/Source/Programs/UnrealBuildTool/Windows/VCToolChain.cs
+++ b/Engine/Source/Programs/UnrealBuildTool/Windows/VCToolChain.cs
@@ -53,6 +53,9 @@ namespace UnrealBuildTool
 				Arguments.Append(" /D_MSC_FULL_VER=" + VersionString + "00000");
 			}
 
+			// @todo clang: Clang on Windows doesn't respect "#pragma warning (error: ####)", and we're not passing "/WX", so warnings are not
+			// treated as errors when compiling on Windows using Clang right now.
+
 
 			if( BuildConfiguration.bEnableCodeAnalysis )
 			{
@@ -75,8 +78,6 @@ namespace UnrealBuildTool
 			// Enable intrinsic functions.
 			Arguments.Append(" /Oi");
 
-			// Enable for static code analysis (where supported). Not treating analysis warnings as errors.
-			// Arguments.Append(" /analyze:WX-");
 			
 			if( WindowsPlatform.bCompileWithClang )
 			{
