Always show the marketplace button if the launcher is installed (or can be installed from the installers in Engine\\Extras\\UnrealEngineLauncher).

diff --git a/Engine/Source/Developer/DesktopPlatform/Private/Linux/DesktopPlatformLinux.cpp b/Engine/Source/Developer/DesktopPlatform/Private/Linux/DesktopPlatformLinux.cpp
--- a/Engine/Source/Developer/DesktopPlatform/Private/Linux/DesktopPlatformLinux.cpp
+++ b/Engine/Source/Developer/DesktopPlatform/Private/Linux/DesktopPlatformLinux.cpp
@@ -87,6 +87,12 @@ bool FDesktopPlatformLinux::OpenFontDialog(const void* ParentWindowHandle, FStri
 	return false;
 }
 
+bool FDesktopPlatformLinux::CanOpenLauncher(bool Install)
+{
+	// TODO: no launcher support at the moment
+	return false;
+}
+
 bool FDesktopPlatformLinux::OpenLauncher(bool Install, FString CommandLineParams )
 {
 	// TODO: support launcher for realz
diff --git a/Engine/Source/Developer/DesktopPlatform/Private/Linux/DesktopPlatformLinux.h b/Engine/Source/Developer/DesktopPlatform/Private/Linux/DesktopPlatformLinux.h
--- a/Engine/Source/Developer/DesktopPlatform/Private/Linux/DesktopPlatformLinux.h
+++ b/Engine/Source/Developer/DesktopPlatform/Private/Linux/DesktopPlatformLinux.h
@@ -13,6 +13,7 @@ public:
 	virtual bool SaveFileDialog(const void* ParentWindowHandle, const FString& DialogTitle, const FString& DefaultPath, const FString& DefaultFile, const FString& FileTypes, uint32 Flags, TArray<FString>& OutFilenames) override;
 	virtual bool OpenDirectoryDialog(const void* ParentWindowHandle, const FString& DialogTitle, const FString& DefaultPath, FString& OutFolderName) override;
 	virtual bool OpenFontDialog(const void* ParentWindowHandle, FString& OutFontName, float& OutHeight, EFontImportFlags& OutFlags) override;
+	virtual bool CanOpenLauncher(bool Install) override;
 	virtual bool OpenLauncher(bool Install, FString CommandLineParams ) override;
 
 	virtual bool RegisterEngineInstallation(const FString &RootDir, FString &OutIdentifier) override;
diff --git a/Engine/Source/Developer/DesktopPlatform/Private/Mac/DesktopPlatformMac.cpp b/Engine/Source/Developer/DesktopPlatform/Private/Mac/DesktopPlatformMac.cpp
--- a/Engine/Source/Developer/DesktopPlatform/Private/Mac/DesktopPlatformMac.cpp
+++ b/Engine/Source/Developer/DesktopPlatform/Private/Mac/DesktopPlatformMac.cpp
@@ -355,6 +355,12 @@ void OpenLauncherCommandLine( const FString& InCommandLine )
 	
 }
 
+bool FDesktopPlatformMac::CanOpenLauncher(bool Install)
+{
+	FString Path;
+	return GetLauncherPath(Path) || (Install && GetLauncherInstallerPath(Path));
+}
+
 bool FDesktopPlatformMac::OpenLauncher(bool Install, FString CommandLineParams )
 {
 	// If the launcher is already running, bring it to front
@@ -367,61 +373,26 @@ bool FDesktopPlatformMac::OpenLauncher(bool Install, FString CommandLineParams )
 		return true;
 	}
 
-	FString LaunchPath;
-
-	bool bWasLaunched = false;
-	if (FParse::Param(FCommandLine::Get(), TEXT("Dev")))
-	{
-		LaunchPath = FPaths::Combine(*FPaths::EngineDir(), TEXT("Binaries"), TEXT("Mac"), TEXT("UnrealEngineLauncher-Mac-Debug.app/Contents/MacOS/UnrealEngineLauncher-Mac-Debug"));
-		CommandLineParams += TEXT(" -noselfupdate");
-	}
-	else
-	{
-		LaunchPath = TEXT("/Applications/Unreal Engine.app/Contents/MacOS/UnrealEngineLauncher-Mac-Shipping");
-	}
-
-	if( !FPaths::FileExists(LaunchPath) )
+	// Try to start a new instance
+	FString LauncherPath;
+	if(GetLauncherPath(LauncherPath))
 	{
-		// The app could have been installed manually from the dmg, launch it where ever it is
-		NSWorkspace* Workspace = [NSWorkspace sharedWorkspace];
-		NSString* Path = [Workspace fullPathForApplication:@"Unreal Engine"];
-		if( Path )
+		if (FParse::Param(FCommandLine::Get(), TEXT("Dev")))
 		{
-			NSURL* URL = [NSURL fileURLWithPath:Path];
-			if( URL )
-			{
-				NSArray* Arguments = [NSArray arrayWithObjects:CommandLineParams.GetNSString(),nil];
-				bWasLaunched = [Workspace launchApplicationAtURL:URL options:0 configuration:[NSDictionary dictionaryWithObject:Arguments forKey:NSWorkspaceLaunchConfigurationArguments] error:nil];
-				
-				OpenLauncherCommandLine(CommandLineParams);
-			}
+			CommandLineParams += TEXT(" -noselfupdate");
 		}
+		return FPlatformProcess::CreateProc(*LauncherPath, *CommandLineParams, true, false, false, NULL, 0, *FPaths::GetPath(LauncherPath), NULL).IsValid();
 	}
 
-	if( !bWasLaunched )
+	// Try to install it
+	FString InstallerPath;
+	if(GetLauncherInstallerPath(InstallerPath))
 	{
-		if( FPaths::FileExists(LaunchPath) )
-		{
-			bWasLaunched = FPlatformProcess::CreateProc(*LaunchPath, *CommandLineParams, true, false, false, NULL, 0, *FPaths::GetPath(LaunchPath), NULL).IsValid();
-		}
-		else if (Install)
-		{
-			// ... or run the installer if desired
-			LaunchPath = FPaths::ConvertRelativePathToFull(FPaths::Combine(*FPaths::EngineDir(), TEXT("Extras/UnrealEngineLauncher/UnrealEngine.dmg")));
-
-			if (FPaths::FileExists(LaunchPath))
-			{
-				FPlatformProcess::LaunchFileInDefaultExternalApplication(*LaunchPath);
-				bWasLaunched = true;
-			}
-			else
-			{
-				bWasLaunched = false;
-			}
-		}
-
+		FPlatformProcess::LaunchFileInDefaultExternalApplication(*InstallerPath);
+		return true;
 	}
-	return bWasLaunched;
+
+	return false;
 }
 
 bool FDesktopPlatformMac::FileDialogShared(bool bSave, const void* ParentWindowHandle, const FString& DialogTitle, const FString& DefaultPath, const FString& DefaultFile, const FString& FileTypes, uint32 Flags, TArray<FString>& OutFilenames)
@@ -665,6 +636,53 @@ bool FDesktopPlatformMac::RunUnrealBuildTool(const FText& Description, const FSt
 	return FFeedbackContextMarkup::PipeProcessOutput(Description, TEXT("/bin/sh"), CmdLineParams, Warn, &ExitCode) && ExitCode == 0;
 }
 
+
+bool FDesktopPlatformMac::GetLauncherPath(FString& OutLauncherPath) const
+{
+	// Try the default executable in the binaries directory
+	if (FParse::Param(FCommandLine::Get(), TEXT("Dev")))
+	{
+		FString LauncherPath = FPaths::Combine(*FPaths::EngineDir(), TEXT("Binaries"), TEXT("Mac"), TEXT("UnrealEngineLauncher-Mac-Debug.app/Contents/MacOS/UnrealEngineLauncher-Mac-Debug"));
+		if(FPaths::FileExists(LauncherPath))
+		{
+			OutLauncherPath = LauncherPath;
+			return true;
+		}
+	}
+	else
+	{
+		FString LauncherPath = TEXT("/Applications/Unreal Engine.app/Contents/MacOS/UnrealEngineLauncher-Mac-Shipping");
+		if(FPaths::FileExists(LauncherPath))
+		{
+			OutLauncherPath = LauncherPath;
+			return true;
+		}
+	}
+
+	// Otherwise search for it...
+	NSWorkspace* Workspace = [NSWorkspace sharedWorkspace];
+	NSString* Path = [Workspace fullPathForApplication:@"Unreal Engine"];
+	if( Path )
+	{
+		OutLauncherPath = FString(Path);
+		return true;
+	}
+
+	return false;
+}
+
+bool FDesktopPlatformMac::GetLauncherInstallerPath(FString& OutInstallerPath) const
+{
+	// Check if the installer exists
+	FString InstallerPath = FPaths::ConvertRelativePathToFull(FPaths::Combine(*FPaths::EngineDir(), TEXT("Extras/UnrealEngineLauncher/UnrealEngine.dmg")));
+	if (FPaths::FileExists(InstallerPath))
+	{
+		OutInstallerPath = InstallerPath;
+		return true;
+	}
+	return false;
+}
+
 FFeedbackContext* FDesktopPlatformMac::GetNativeFeedbackContext()
 {
 	static FMacNativeFeedbackContext Warn;
diff --git a/Engine/Source/Developer/DesktopPlatform/Private/Mac/DesktopPlatformMac.h b/Engine/Source/Developer/DesktopPlatform/Private/Mac/DesktopPlatformMac.h
--- a/Engine/Source/Developer/DesktopPlatform/Private/Mac/DesktopPlatformMac.h
+++ b/Engine/Source/Developer/DesktopPlatform/Private/Mac/DesktopPlatformMac.h
@@ -13,6 +13,7 @@ public:
 	virtual bool SaveFileDialog(const void* ParentWindowHandle, const FString& DialogTitle, const FString& DefaultPath, const FString& DefaultFile, const FString& FileTypes, uint32 Flags, TArray<FString>& OutFilenames) override;
 	virtual bool OpenDirectoryDialog(const void* ParentWindowHandle, const FString& DialogTitle, const FString& DefaultPath, FString& OutFolderName) override;
 	virtual bool OpenFontDialog(const void* ParentWindowHandle, FString& OutFontName, float& OutHeight, EFontImportFlags& OutFlags) override;
+	virtual bool CanOpenLauncher(bool Install) override;
 	virtual bool OpenLauncher(bool Install, FString CommandLineParams ) override;
 
 	virtual bool RegisterEngineInstallation(const FString &RootDir, FString &OutIdentifier) override;
@@ -26,6 +27,8 @@ public:
 	virtual FFeedbackContext* GetNativeFeedbackContext() override;
 
 private:
+	bool GetLauncherPath(FString& OutLauncherPath) const;
+	bool GetLauncherInstallerPath(FString& OutInstallerPath) const;
 	bool FileDialogShared(bool bSave, const void* ParentWindowHandle, const FString& DialogTitle, const FString& DefaultPath, const FString& DefaultFile, const FString& FileTypes, uint32 Flags, TArray<FString>& OutFilenames);
 };
 
diff --git a/Engine/Source/Developer/DesktopPlatform/Private/Windows/DesktopPlatformWindows.cpp b/Engine/Source/Developer/DesktopPlatform/Private/Windows/DesktopPlatformWindows.cpp
--- a/Engine/Source/Developer/DesktopPlatform/Private/Windows/DesktopPlatformWindows.cpp
+++ b/Engine/Source/Developer/DesktopPlatform/Private/Windows/DesktopPlatformWindows.cpp
@@ -141,55 +141,84 @@ bool FDesktopPlatformWindows::OpenFontDialog(const void* ParentWindowHandle, FSt
 	return bSuccess;
 }
 
-bool FDesktopPlatformWindows::OpenLauncher(bool Install, FString CommandLineParams )
+bool FDesktopPlatformWindows::CanOpenLauncher(bool Install)
 {
-	FString LaunchPath;
+	// Check if the launcher exists, or (optionally) if the installer exists
+	FString Path;
+	return GetLauncherPath(Path) || (Install && GetLauncherInstallerPath(Path));
+}
 
-	if (FParse::Param(FCommandLine::Get(), TEXT("Dev")))
+bool FDesktopPlatformWindows::OpenLauncher(bool Install, FString CommandLineParams )
+{
+	// Try to launch it directly
+	FString LauncherPath;
+	if(GetLauncherPath(LauncherPath))
 	{
-		LaunchPath = FPaths::Combine(*FPaths::EngineDir(), TEXT("Binaries"), TEXT("Win64"), TEXT("UnrealEngineLauncher-Win64-Debug.exe"));
-		CommandLineParams += TEXT(" -noselfupdate");
+		if (FParse::Param(FCommandLine::Get(), TEXT("Dev")))
+		{
+			CommandLineParams += TEXT(" -noselfupdate");
+		}
+		return FPlatformProcess::CreateProc(*LauncherPath, *CommandLineParams, true, false, false, NULL, 0, *FPaths::GetPath(LauncherPath), NULL).IsValid();
 	}
-	else
+
+	// Otherwise see if we can install it
+	FString InstallerPath;
+	if(Install && GetLauncherInstallerPath(InstallerPath))
 	{
-		LaunchPath = FPaths::Combine(*FPaths::EngineDir(), TEXT("../../Launcher/Engine/Binaries/Win64/UnrealEngineLauncher.exe"));
+		FPlatformProcess::LaunchFileInDefaultExternalApplication(*InstallerPath);
+		return true;
 	}
 
-	// if the Launcher doesn't exist locally...
-	if (!FPaths::FileExists(LaunchPath))
-	{
-		FString InstallDir;
+	return false;
+}
 
-		// ... locate it in the Registry...
-		if (FWindowsPlatformMisc::QueryRegKey(HKEY_LOCAL_MACHINE, TEXT("Software\\EpicGames\\Unreal Engine"), TEXT("InstallDir"), InstallDir) && (InstallDir.Len() > 0))
+bool FDesktopPlatformWindows::GetLauncherPath(FString& OutLauncherPath)
+{
+	// Try the default executable in the binaries directory
+	if (FParse::Param(FCommandLine::Get(), TEXT("Dev")))
+	{
+		FString LauncherPath = FPaths::Combine(*FPaths::EngineDir(), TEXT("Binaries"), TEXT("Win64"), TEXT("UnrealEngineLauncher-Win64-Debug.exe"));
+		if(FPaths::FileExists(LauncherPath))
 		{
-			LaunchPath = FPaths::Combine(*InstallDir, TEXT("Launcher/Engine/Binaries/Win64/UnrealEngineLauncher.exe"));
+			OutLauncherPath = LauncherPath;
+			return true;
 		}
-		else
+	}
+	else
+	{
+		FString LauncherPath = FPaths::Combine(*FPaths::EngineDir(), TEXT("../../Launcher/Engine/Binaries/Win64/UnrealEngineLauncher.exe"));
+		if(FPaths::FileExists(LauncherPath))
 		{
-			LaunchPath.Empty();
+			OutLauncherPath = LauncherPath;
+			return true;
 		}
+	}
 
-		if (LaunchPath.IsEmpty() || !FPaths::FileExists(LaunchPath))
+	// Otherwise locate it in the Registry...
+	FString InstallDir;
+	if (FWindowsPlatformMisc::QueryRegKey(HKEY_LOCAL_MACHINE, TEXT("Software\\EpicGames\\Unreal Engine"), TEXT("InstallDir"), InstallDir) && (InstallDir.Len() > 0))
+	{
+		FString LauncherPath = FPaths::Combine(*InstallDir, TEXT("Launcher/Engine/Binaries/Win64/UnrealEngineLauncher.exe"));
+		if(FPaths::FileExists(LauncherPath))
 		{
-			if (Install)
-			{
-				// ... or run the installer if desired
-				LaunchPath = FPaths::ConvertRelativePathToFull(FPaths::Combine(*FPaths::EngineDir(), TEXT("Extras/UnrealEngineLauncher/UnrealEngineInstaller.msi")));
-
-				if (FPaths::FileExists(LaunchPath))
-				{
-					FPlatformProcess::LaunchFileInDefaultExternalApplication(*LaunchPath);
-
-					return true;
-				}
-			}
-
-			return false;
+			OutLauncherPath = LauncherPath;
+			return true;
 		}
 	}
 
-	return FPlatformProcess::CreateProc(*LaunchPath, *CommandLineParams, true, false, false, NULL, 0, *FPaths::GetPath(LaunchPath), NULL).IsValid();
+	// Otherwise fail
+	return false;
+}
+
+bool FDesktopPlatformWindows::GetLauncherInstallerPath(FString& OutInstallerPath)
+{
+	FString InstallerPath = FPaths::ConvertRelativePathToFull(FPaths::Combine(*FPaths::EngineDir(), TEXT("Extras/UnrealEngineLauncher/UnrealEngineInstaller.msi")));
+	if(FPaths::FileExists(InstallerPath))
+	{
+		OutInstallerPath = InstallerPath;
+		return true;
+	}
+	return false;
 }
 
 bool FDesktopPlatformWindows::FileDialogShared(bool bSave, const void* ParentWindowHandle, const FString& DialogTitle, const FString& DefaultPath, const FString& DefaultFile, const FString& FileTypes, uint32 Flags, TArray<FString>& OutFilenames, int32& OutFilterIndex)
diff --git a/Engine/Source/Developer/DesktopPlatform/Private/Windows/DesktopPlatformWindows.h b/Engine/Source/Developer/DesktopPlatform/Private/Windows/DesktopPlatformWindows.h
--- a/Engine/Source/Developer/DesktopPlatform/Private/Windows/DesktopPlatformWindows.h
+++ b/Engine/Source/Developer/DesktopPlatform/Private/Windows/DesktopPlatformWindows.h
@@ -14,6 +14,7 @@ public:
 	virtual bool SaveFileDialog(const void* ParentWindowHandle, const FString& DialogTitle, const FString& DefaultPath, const FString& DefaultFile, const FString& FileTypes, uint32 Flags, TArray<FString>& OutFilenames) override;
 	virtual bool OpenDirectoryDialog(const void* ParentWindowHandle, const FString& DialogTitle, const FString& DefaultPath, FString& OutFolderName) override;
 	virtual bool OpenFontDialog(const void* ParentWindowHandle, FString& OutFontName, float& OutHeight, EFontImportFlags& OutFlags) override;
+	virtual bool CanOpenLauncher(bool Install) override;
 	virtual bool OpenLauncher(bool Install, FString CommandLineParams ) override;
 
 	virtual bool RegisterEngineInstallation(const FString &RootDir, FString &OutIdentifier) override;
@@ -31,6 +32,8 @@ public:
 	virtual FFeedbackContext* GetNativeFeedbackContext() override;
 
 private:
+	bool GetLauncherPath(FString& OutLauncherPath);
+	bool GetLauncherInstallerPath(FString& OutInstallerPath);
 	bool FileDialogShared(bool bSave, const void* ParentWindowHandle, const FString& DialogTitle, const FString& DefaultPath, const FString& DefaultFile, const FString& FileTypes, uint32 Flags, TArray<FString>& OutFilenames, int32& OutFilterIndex);
 	void GetRequiredRegistrySettings(TIndirectArray<FRegistryRootedKey> &RootedKeys);
 	int32 GetShellIntegrationVersion(const FString &FileName);
diff --git a/Engine/Source/Developer/DesktopPlatform/Public/IDesktopPlatform.h b/Engine/Source/Developer/DesktopPlatform/Public/IDesktopPlatform.h
--- a/Engine/Source/Developer/DesktopPlatform/Public/IDesktopPlatform.h
+++ b/Engine/Source/Developer/DesktopPlatform/Public/IDesktopPlatform.h
@@ -123,6 +123,14 @@ public:
 	 */
 	virtual bool OpenFontDialog(const void* ParentWindowHandle, FString& OutFontName, float& OutHeight, EFontImportFlags& OutFlags) = 0;
 
+	/**
+	 * Determines whether the launcher can be opened.
+	 *
+	 * @param Install					Whether to include the possibility of installing the launcher in the check.
+	 * @return true if the launcher can be opened (or installed).
+	 */
+	virtual bool CanOpenLauncher(bool Install) = 0;
+
 	/**
 	 * Opens the marketplace user interface.
 	 *
diff --git a/Engine/Source/Editor/GameProjectGeneration/Private/SGameProjectDialog.cpp b/Engine/Source/Editor/GameProjectGeneration/Private/SGameProjectDialog.cpp
--- a/Engine/Source/Editor/GameProjectGeneration/Private/SGameProjectDialog.cpp
+++ b/Engine/Source/Editor/GameProjectGeneration/Private/SGameProjectDialog.cpp
@@ -123,7 +123,7 @@ void SGameProjectDialog::Construct( const FArguments& InArgs )
 											.Padding(3)
 											[
 												SAssignNew(MarketplaceButton, SButton)
-													.Visibility( (FEngineBuildSettings::IsPerforceBuild() && !FEngineBuildSettings::IsInternalBuild()) ? EVisibility::Collapsed : EVisibility::Visible )
+													.Visibility( FDesktopPlatformModule::Get()->CanOpenLauncher(true) ? EVisibility::Collapsed : EVisibility::Visible )
 													.ForegroundColor(this, &SGameProjectDialog::HandleMarketplaceTabButtonForegroundColor)
 													.ButtonStyle(FEditorStyle::Get(), "ToggleButton")
 													.OnClicked(this, &SGameProjectDialog::HandleMarketplaceTabButtonClicked)
diff --git a/Engine/Source/Editor/LevelEditor/Private/LevelEditorToolBar.cpp b/Engine/Source/Editor/LevelEditor/Private/LevelEditorToolBar.cpp
--- a/Engine/Source/Editor/LevelEditor/Private/LevelEditorToolBar.cpp
+++ b/Engine/Source/Editor/LevelEditor/Private/LevelEditorToolBar.cpp
@@ -43,9 +43,7 @@ TSharedRef< SWidget > FLevelEditorToolBar::MakeLevelEditorToolBar( const TShared
 	ToolbarBuilder.BeginSection("Content");
 	{
 		ToolbarBuilder.AddToolBarButton( FLevelEditorCommands::Get().OpenContentBrowser, NAME_None, LOCTEXT( "ContentBrowser_Override", "Content" ), TAttribute<FText>(), TAttribute<FSlateIcon>(), "LevelToolbarContent" );
-
-		// Internal builds and external builds should see the marketplace by default, perforce builds should not.
-		if (FEngineBuildSettings::IsInternalBuild() || !FEngineBuildSettings::IsPerforceBuild()) 
+		if (FDesktopPlatformModule::Get()->CanOpenLauncher(true)) 
 		{
 			ToolbarBuilder.AddToolBarButton(FLevelEditorCommands::Get().OpenMarketplace, NAME_None, LOCTEXT("Marketplace_Override", "Marketplace"), TAttribute<FText>(), TAttribute<FSlateIcon>(), "LevelToolbarMarketplace");
 		}
