Stop calling some game-thread code on the main-thread when opening the native open/save/font dialogs on OS X post thread separation.

diff --git a/Engine/Source/Developer/DesktopPlatform/Private/Mac/DesktopPlatformMac.cpp b/Engine/Source/Developer/DesktopPlatform/Private/Mac/DesktopPlatformMac.cpp
--- a/Engine/Source/Developer/DesktopPlatform/Private/Mac/DesktopPlatformMac.cpp
+++ b/Engine/Source/Developer/DesktopPlatform/Private/Mac/DesktopPlatformMac.cpp
@@ -226,29 +226,30 @@ bool FDesktopPlatformMac::OpenDirectoryDialog(const void* ParentWindowHandle, co
 {
 	MacApplication->SetCapture( NULL );
 
-	bool bSuccess = MainThreadReturn(^{
-		SCOPED_AUTORELEASE_POOL;
-		FCocoaScopeContext ContextGuard;
+	bool bSuccess = false;
+	{
+		FScopedSystemModalMode SystemModalScope;
+		bSuccess = MainThreadReturn(^{
+			SCOPED_AUTORELEASE_POOL;
+			FCocoaScopeContext ContextGuard;
 
-		NSOpenPanel* Panel = [NSOpenPanel openPanel];
-		[Panel setCanChooseFiles: false];
-		[Panel setCanChooseDirectories: true];
-		[Panel setAllowsMultipleSelection: false];
-		[Panel setCanCreateDirectories: true];
+			NSOpenPanel* Panel = [NSOpenPanel openPanel];
+			[Panel setCanChooseFiles: false];
+			[Panel setCanChooseDirectories: true];
+			[Panel setAllowsMultipleSelection: false];
+			[Panel setCanCreateDirectories: true];
 
-		CFStringRef Title = FPlatformString::TCHARToCFString(*DialogTitle);
-		[Panel setTitle: (NSString*)Title];
-		CFRelease(Title);
+			CFStringRef Title = FPlatformString::TCHARToCFString(*DialogTitle);
+			[Panel setTitle: (NSString*)Title];
+			CFRelease(Title);
 
-		CFStringRef DefaultPathCFString = FPlatformString::TCHARToCFString(*DefaultPath);
-		NSURL* DefaultPathURL = [NSURL fileURLWithPath: (NSString*)DefaultPathCFString];
-		[Panel setDirectoryURL: DefaultPathURL];
-		CFRelease(DefaultPathCFString);
+			CFStringRef DefaultPathCFString = FPlatformString::TCHARToCFString(*DefaultPath);
+			NSURL* DefaultPathURL = [NSURL fileURLWithPath: (NSString*)DefaultPathCFString];
+			[Panel setDirectoryURL: DefaultPathURL];
+			CFRelease(DefaultPathCFString);
 
-		bool bSuccess = false;
+			bool bSuccess = false;
 
-		{
-			FScopedSystemModalMode SystemModalScope;
 			NSInteger Result = [Panel runModal];
 
 			if (Result == NSFileHandlingPanelOKButton)
@@ -261,13 +262,12 @@ bool FDesktopPlatformMac::OpenDirectoryDialog(const void* ParentWindowHandle, co
 
 				bSuccess = true;
 			}
-		}
 
-		[Panel close];
+			[Panel close];
 
-		return bSuccess;
-	});
-	
+			return bSuccess;
+		});
+	}
 	MacApplication->ResetModifierKeys();
 	
 	return bSuccess;
@@ -277,51 +277,52 @@ bool FDesktopPlatformMac::OpenFontDialog(const void* ParentWindowHandle, FString
 {
 	MacApplication->SetCapture( NULL );
 	
-	bool bSuccess = MainThreadReturn(^{
-		SCOPED_AUTORELEASE_POOL;
-		FCocoaScopeContext ContextGuard;
+	bool bSuccess = false;
+	{
+		FScopedSystemModalMode SystemModalScope;
+		bSuccess = MainThreadReturn(^{
+			SCOPED_AUTORELEASE_POOL;
+			FCocoaScopeContext ContextGuard;
 
-		NSFontPanel* Panel = [NSFontPanel sharedFontPanel];
-		[Panel setFloatingPanel: false];
-		[[Panel standardWindowButton: NSWindowCloseButton] setEnabled: false];
-		FFontDialogAccessoryView* AccessoryView = [[FFontDialogAccessoryView alloc] initWithFrame: NSMakeRect(0.0, 0.0, 190.0, 80.0)];
-		[Panel setAccessoryView: AccessoryView];
+			NSFontPanel* Panel = [NSFontPanel sharedFontPanel];
+			[Panel setFloatingPanel: false];
+			[[Panel standardWindowButton: NSWindowCloseButton] setEnabled: false];
+			FFontDialogAccessoryView* AccessoryView = [[FFontDialogAccessoryView alloc] initWithFrame: NSMakeRect(0.0, 0.0, 190.0, 80.0)];
+			[Panel setAccessoryView: AccessoryView];
 
-		{
-			FScopedSystemModalMode SystemModalScope;
 			[NSApp runModalForWindow: Panel];
-		}
+			
+			[Panel close];
 
-		[Panel close];
+			bool bSuccess = [AccessoryView result];
 
-		bool bSuccess = [AccessoryView result];
+			[Panel setAccessoryView: NULL];
+			[AccessoryView release];
+			[[Panel standardWindowButton: NSWindowCloseButton] setEnabled: true];
 
-		[Panel setAccessoryView: NULL];
-		[AccessoryView release];
-		[[Panel standardWindowButton: NSWindowCloseButton] setEnabled: true];
+			if( bSuccess )
+			{
+				NSFont* Font = [Panel panelConvertFont: [NSFont userFontOfSize: 0]];
 
-		if( bSuccess )
-		{
-			NSFont* Font = [Panel panelConvertFont: [NSFont userFontOfSize: 0]];
+				TCHAR FontName[MAX_PATH];
+				FPlatformString::CFStringToTCHAR((CFStringRef)[Font fontName], FontName);
 
-			TCHAR FontName[MAX_PATH];
-			FPlatformString::CFStringToTCHAR((CFStringRef)[Font fontName], FontName);
+				OutFontName = FontName;
+				OutHeight = [Font pointSize];
 
-			OutFontName = FontName;
-			OutHeight = [Font pointSize];
+				auto FontFlags = EFontImportFlags::None;
 
-			auto FontFlags = EFontImportFlags::None;
+				if( [Font underlineThickness] >= 1.0 )
+				{
+					FontFlags |= EFontImportFlags::EnableUnderline;
+				}
 
-			if( [Font underlineThickness] >= 1.0 )
-			{
-				FontFlags |= EFontImportFlags::EnableUnderline;
+				OutFlags = FontFlags;
 			}
 
-			OutFlags = FontFlags;
-		}
-
-		return bSuccess;
-	});
+			return bSuccess;
+		});
+	}
 	
 	MacApplication->ResetModifierKeys();
 	
@@ -408,60 +409,60 @@ bool FDesktopPlatformMac::FileDialogShared(bool bSave, const void* ParentWindowH
 {
 	MacApplication->SetCapture( NULL );
 
-	bool bSuccess = MainThreadReturn(^{
-		SCOPED_AUTORELEASE_POOL;
-		FCocoaScopeContext ContextGuard;
+	bool bSuccess = false;
+	{
+		FScopedSystemModalMode SystemModalScope;
+		bSuccess = MainThreadReturn(^{
+			SCOPED_AUTORELEASE_POOL;
+			FCocoaScopeContext ContextGuard;
 
-		NSSavePanel* Panel = bSave ? [NSSavePanel savePanel] : [NSOpenPanel openPanel];
+			NSSavePanel* Panel = bSave ? [NSSavePanel savePanel] : [NSOpenPanel openPanel];
 
-		if (!bSave)
-		{
-			NSOpenPanel* OpenPanel = (NSOpenPanel*)Panel;
-			[OpenPanel setCanChooseFiles: true];
-			[OpenPanel setCanChooseDirectories: false];
-			[OpenPanel setAllowsMultipleSelection: Flags & EFileDialogFlags::Multiple];
-		}
+			if (!bSave)
+			{
+				NSOpenPanel* OpenPanel = (NSOpenPanel*)Panel;
+				[OpenPanel setCanChooseFiles: true];
+				[OpenPanel setCanChooseDirectories: false];
+				[OpenPanel setAllowsMultipleSelection: Flags & EFileDialogFlags::Multiple];
+			}
 
-		[Panel setCanCreateDirectories: bSave];
+			[Panel setCanCreateDirectories: bSave];
 
-		CFStringRef Title = FPlatformString::TCHARToCFString(*DialogTitle);
-		[Panel setTitle: (NSString*)Title];
-		CFRelease(Title);
+			CFStringRef Title = FPlatformString::TCHARToCFString(*DialogTitle);
+			[Panel setTitle: (NSString*)Title];
+			CFRelease(Title);
 
-		CFStringRef DefaultPathCFString = FPlatformString::TCHARToCFString(*DefaultPath);
-		NSURL* DefaultPathURL = [NSURL fileURLWithPath: (NSString*)DefaultPathCFString];
-		[Panel setDirectoryURL: DefaultPathURL];
-		CFRelease(DefaultPathCFString);
+			CFStringRef DefaultPathCFString = FPlatformString::TCHARToCFString(*DefaultPath);
+			NSURL* DefaultPathURL = [NSURL fileURLWithPath: (NSString*)DefaultPathCFString];
+			[Panel setDirectoryURL: DefaultPathURL];
+			CFRelease(DefaultPathCFString);
 
-		CFStringRef FileNameCFString = FPlatformString::TCHARToCFString(*DefaultFile);
-		[Panel setNameFieldStringValue: (NSString*)FileNameCFString];
-		CFRelease(FileNameCFString);
+			CFStringRef FileNameCFString = FPlatformString::TCHARToCFString(*DefaultFile);
+			[Panel setNameFieldStringValue: (NSString*)FileNameCFString];
+			CFRelease(FileNameCFString);
 
-		FFileDialogAccessoryView* AccessoryView = [[FFileDialogAccessoryView alloc] initWithFrame: NSMakeRect( 0.0, 0.0, 250.0, 85.0 ) dialogPanel: Panel];
-		[Panel setAccessoryView: AccessoryView];
+			FFileDialogAccessoryView* AccessoryView = [[FFileDialogAccessoryView alloc] initWithFrame: NSMakeRect( 0.0, 0.0, 250.0, 85.0 ) dialogPanel: Panel];
+			[Panel setAccessoryView: AccessoryView];
 
-		TArray<FString> FileTypesArray;
-		int32 NumFileTypes = FileTypes.ParseIntoArray(&FileTypesArray, TEXT("|"), true);
+			TArray<FString> FileTypesArray;
+			int32 NumFileTypes = FileTypes.ParseIntoArray(&FileTypesArray, TEXT("|"), true);
 
-		NSMutableArray* AllowedFileTypes = [NSMutableArray arrayWithCapacity: NumFileTypes];
+			NSMutableArray* AllowedFileTypes = [NSMutableArray arrayWithCapacity: NumFileTypes];
 
-		if( NumFileTypes > 0 )
-		{
-			for( int32 Index = 0; Index < NumFileTypes; ++Index )
+			if( NumFileTypes > 0 )
 			{
-				CFStringRef Type = FPlatformString::TCHARToCFString(*FileTypesArray[Index]);
-				[AllowedFileTypes addObject: (NSString*)Type];
-				CFRelease(Type);
+				for( int32 Index = 0; Index < NumFileTypes; ++Index )
+				{
+					CFStringRef Type = FPlatformString::TCHARToCFString(*FileTypesArray[Index]);
+					[AllowedFileTypes addObject: (NSString*)Type];
+					CFRelease(Type);
+				}
 			}
-		}
-
-		[AccessoryView AddAllowedFileTypes:AllowedFileTypes];
 
-		bool bSuccess = false;
-		NSWindow* FocusWindow = [[NSApplication sharedApplication] keyWindow];
+			[AccessoryView AddAllowedFileTypes:AllowedFileTypes];
 
-		{
-			FScopedSystemModalMode SystemModalScope;
+			bool bSuccess = false;
+			NSWindow* FocusWindow = [[NSApplication sharedApplication] keyWindow];
 
 			NSInteger Result = [Panel runModal];
 			[AccessoryView release];
@@ -493,17 +494,17 @@ bool FDesktopPlatformMac::FileDialogShared(bool bSave, const void* ParentWindowH
 
 				bSuccess = true;
 			}
-		}
 
-		[Panel close];
-		
-		if(FocusWindow)
-		{
-			[FocusWindow makeKeyWindow];
-		}
+			[Panel close];
+			
+			if(FocusWindow)
+			{
+				[FocusWindow makeKeyWindow];
+			}
 
-		return bSuccess;
-	});
+			return bSuccess;
+		});
+	}
 	
 	MacApplication->ResetModifierKeys();
 
