Fixed stale file association settings that were never cleaned up from the old UE/Rocket shell extensions.

diff --git a/Engine/Source/Developer/DesktopPlatform/Private/Windows/DesktopPlatformWindows.cpp b/Engine/Source/Developer/DesktopPlatform/Private/Windows/DesktopPlatformWindows.cpp
--- a/Engine/Source/Developer/DesktopPlatform/Private/Windows/DesktopPlatformWindows.cpp
+++ b/Engine/Source/Developer/DesktopPlatform/Private/Windows/DesktopPlatformWindows.cpp
@@ -575,6 +575,10 @@ void FDesktopPlatformWindows::GetRequiredRegistrySettings(TIndirectArray<FRegist
 	FPaths::MakePlatformFilename(ExecutableFileName);
 	FString QuotedExecutableFileName = FString(TEXT("\"")) + ExecutableFileName + FString(TEXT("\""));
 
+	// HKCU\SOFTWARE\Classes\.uproject
+	FRegistryRootedKey *UserRootExtensionKey = new FRegistryRootedKey(HKEY_CURRENT_USER, TEXT("SOFTWARE\\Classes\\.uproject"));
+	RootedKeys.AddRawItem(UserRootExtensionKey);
+
 	// HKLM\SOFTWARE\Classes\.uproject
 	FRegistryRootedKey *RootExtensionKey = new FRegistryRootedKey(HKEY_LOCAL_MACHINE, TEXT("SOFTWARE\\Classes\\.uproject"));
 	RootExtensionKey->Key = new FRegistryKey();
@@ -613,6 +617,18 @@ void FDesktopPlatformWindows::GetRequiredRegistrySettings(TIndirectArray<FRegist
 	ShellVersionKey->SetValue(TEXT(""), TEXT("Switch Unreal Engine version..."));
 	ShellVersionKey->SetValue(TEXT("Icon"), QuotedExecutableFileName);
 	ShellVersionKey->FindOrAddKey(L"command")->SetValue(TEXT(""), QuotedExecutableFileName + FString(TEXT(" /switchversion \"%1\"")));
+
+	// If the user has manually selected something other than our extension, we need to delete it. Explorer explicitly disables set access on that keys in that folder, but we can delete the whole thing.
+	TCHAR ProgId[128];
+	::DWORD ProgIdSize = sizeof(ProgId);
+	if(RegGetValue(HKEY_CURRENT_USER, TEXT("Software\\Microsoft\\Windows\\CurrentVersion\\Explorer\\FileExts\\.uproject\\UserChoice"), TEXT("Progid"), RRF_RT_REG_SZ, NULL, ProgId, &ProgIdSize) != S_OK)
+	{
+		ProgId[0] = 0;
+	}
+	if(FCString::Strcmp(ProgId, TEXT("Unreal.ProjectFile")) != 0)
+	{
+		RootedKeys.AddRawItem(new FRegistryRootedKey(HKEY_CURRENT_USER, TEXT("Software\\Microsoft\\Windows\\CurrentVersion\\Explorer\\FileExts\\.uproject\\UserChoice")));
+	}
 }
 
 int32 FDesktopPlatformWindows::GetShellIntegrationVersion(const FString &FileName)
