Changed FDesktopPlatformBase::IsStockEngineRelease() to use FGuid::Parse() and updated Mac engine id generator to use EGuidFormats::DigitsWithHyphensInBraces

diff --git a/Engine/Source/Developer/DesktopPlatform/Private/DesktopPlatformBase.cpp b/Engine/Source/Developer/DesktopPlatform/Private/DesktopPlatformBase.cpp
--- a/Engine/Source/Developer/DesktopPlatform/Private/DesktopPlatformBase.cpp
+++ b/Engine/Source/Developer/DesktopPlatform/Private/DesktopPlatformBase.cpp
@@ -190,7 +190,8 @@ bool FDesktopPlatformBase::IsPreferredEngineIdentifier(const FString &Identifier
 
 bool FDesktopPlatformBase::IsStockEngineRelease(const FString &Identifier)
 {
-	return Identifier.Len() > 0 && FChar::IsDigit(Identifier[0]);
+	FGuid Guid;
+	return !FGuid::Parse(Identifier, Guid);
 }
 
 bool FDesktopPlatformBase::IsSourceDistribution(const FString &EngineRootDir)
diff --git a/Engine/Source/Developer/DesktopPlatform/Private/Mac/DesktopPlatformMac.cpp b/Engine/Source/Developer/DesktopPlatform/Private/Mac/DesktopPlatformMac.cpp
--- a/Engine/Source/Developer/DesktopPlatform/Private/Mac/DesktopPlatformMac.cpp
+++ b/Engine/Source/Developer/DesktopPlatform/Private/Mac/DesktopPlatformMac.cpp
@@ -589,7 +589,7 @@ void FDesktopPlatformMac::EnumerateEngineInstallations(TMap<FString, FString> &O
 				{
 					if (!OutInstallations.FindKey(EngineDir))
 					{
-						EngineId = FGuid::NewGuid().ToString(EGuidFormats::DigitsWithHyphens);
+						EngineId = FGuid::NewGuid().ToString(EGuidFormats::DigitsWithHyphensInBraces);
 						Section.AddUnique(*EngineId, EngineDir);
 						ConfigFile.Dirty = true;
 					}
