Fix DesktopPlatform module on Linux.

diff --git a/Engine/Source/Developer/DesktopPlatform/Private/DesktopPlatformBase.cpp b/Engine/Source/Developer/DesktopPlatform/Private/DesktopPlatformBase.cpp
--- a/Engine/Source/Developer/DesktopPlatform/Private/DesktopPlatformBase.cpp
+++ b/Engine/Source/Developer/DesktopPlatform/Private/DesktopPlatformBase.cpp
@@ -440,7 +440,7 @@ bool FDesktopPlatformBase::GetSolutionPath(FString& OutSolutionPath)
 #if PLATFORM_MAC
 	const TCHAR* Suffix = TEXT(".xcodeproj/project.pbxproj");
 #elif PLATFORM_LINUX
-	UE_LOG(LogModuleManager, Warning, TEXT("STUBBED: solution file path for Linux"));
+	UE_LOG(LogDesktopPlatform, Warning, TEXT("STUBBED: solution file path for Linux"));
 	const TCHAR* Suffix = TEXT("/stubbed/path/to.solution");
 #else
 	const TCHAR* Suffix = TEXT(".sln");
