Fixed a crash caused by FMacNativeFeedbackContextWindowController's window being released while still in use

diff --git a/Engine/Source/Developer/DesktopPlatform/Private/Mac/MacNativeFeedbackContext.cpp b/Engine/Source/Developer/DesktopPlatform/Private/Mac/MacNativeFeedbackContext.cpp
--- a/Engine/Source/Developer/DesktopPlatform/Private/Mac/MacNativeFeedbackContext.cpp
+++ b/Engine/Source/Developer/DesktopPlatform/Private/Mac/MacNativeFeedbackContext.cpp
@@ -58,7 +58,7 @@
 		
 		NSRect WindowRect = NSMakeRect(ConsolePosX, ConsolePosY, ConsoleWidth, ConsoleHeight);
 		
-		Window = [[[NSWindow alloc] initWithContentRect:WindowRect styleMask:NSTitledWindowMask|NSMiniaturizableWindowMask|NSResizableWindowMask|NSClosableWindowMask backing:NSBackingStoreBuffered defer:NO] autorelease];
+		Window = [[NSWindow alloc] initWithContentRect:WindowRect styleMask:NSTitledWindowMask|NSMiniaturizableWindowMask|NSResizableWindowMask|NSClosableWindowMask backing:NSBackingStoreBuffered defer:NO];
 		[Window setTitle:@"Unreal Engine 4"];
 		[Window setReleasedWhenClosed:NO];
 		[Window setMinSize:NSMakeSize(400, 100)];
