Core - Moved GetMachineId/GetEpicAccountId to FGenericPlatformMisc
CrashReport - Restored missing functionality on Linux

diff --git a/Engine/Source/Developer/DesktopPlatform/Private/DesktopPlatformBase.cpp b/Engine/Source/Developer/DesktopPlatform/Private/DesktopPlatformBase.cpp
--- a/Engine/Source/Developer/DesktopPlatform/Private/DesktopPlatformBase.cpp
+++ b/Engine/Source/Developer/DesktopPlatform/Private/DesktopPlatformBase.cpp
@@ -475,40 +475,6 @@ FString FDesktopPlatformBase::GetDefaultProjectCreationPath()
 	return FString(FPlatformProcess::UserDir()) + DefaultProjectSubFolder;
 }
 
-FGuid FDesktopPlatformBase::GetMachineId()
-{
-	FGuid MachineId;
-	FString MachineIdStr;
-
-	// Check to see if we already have a valid machine ID to use
-	if(!FPlatformMisc::GetStoredValue(TEXT("Epic Games"), TEXT("Unreal Engine/Identifiers"), TEXT("MachineId"), MachineIdStr) || !FGuid::Parse(MachineIdStr, MachineId))
-	{
-		// No valid machine ID, generate and save a new one
-		MachineId = FGuid::NewGuid();
-		MachineIdStr = MachineId.ToString(EGuidFormats::Digits);
-
-		if(!FPlatformMisc::SetStoredValue(TEXT("Epic Games"), TEXT("Unreal Engine/Identifiers"), TEXT("MachineId"), MachineIdStr))
-		{
-			// Failed to persist the machine ID - reset it to zero to avoid returning a transient value
-			MachineId = FGuid();
-		}
-	}
-
-	return MachineId;
-}
-
-FString FDesktopPlatformBase::GetEpicAccountId()
-{
-	FString AccountId;
-	FPlatformMisc::GetStoredValue(TEXT("Epic Games"), TEXT("Unreal Engine/Identifiers"), TEXT("AccountId"), AccountId);
-	return AccountId;
-}
-
-void FDesktopPlatformBase::SetEpicAccountId(const FString& AccountId)
-{
-	FPlatformMisc::SetStoredValue(TEXT("Epic Games"), TEXT("Unreal Engine/Identifiers"), TEXT("AccountId"), AccountId);
-}
-
 void FDesktopPlatformBase::ReadLauncherInstallationList()
 {
 	FString InstalledListFile = FString(FPlatformProcess::ApplicationSettingsDir()) / TEXT("UnrealEngineLauncher/LauncherInstalled.dat");
diff --git a/Engine/Source/Developer/DesktopPlatform/Private/DesktopPlatformBase.h b/Engine/Source/Developer/DesktopPlatform/Private/DesktopPlatformBase.h
--- a/Engine/Source/Developer/DesktopPlatform/Private/DesktopPlatformBase.h
+++ b/Engine/Source/Developer/DesktopPlatform/Private/DesktopPlatformBase.h
@@ -44,10 +44,6 @@ public:
 	virtual bool EnumerateProjectsKnownByEngine(const FString &Identifier, bool bIncludeNativeProjects, TArray<FString> &OutProjectFileNames) override;
 	virtual FString GetDefaultProjectCreationPath() override;
 
-	virtual FGuid GetMachineId() override;
-	virtual FString GetEpicAccountId() override;
-	virtual void SetEpicAccountId(const FString& AccountId) override;
-
 private:
 	FString CurrentEngineIdentifier;
 	FDateTime LauncherInstallationTimestamp;
diff --git a/Engine/Source/Developer/DesktopPlatform/Public/IDesktopPlatform.h b/Engine/Source/Developer/DesktopPlatform/Public/IDesktopPlatform.h
--- a/Engine/Source/Developer/DesktopPlatform/Public/IDesktopPlatform.h
+++ b/Engine/Source/Developer/DesktopPlatform/Public/IDesktopPlatform.h
@@ -370,19 +370,4 @@ public:
 	* @return Path to the folder
 	*/
 	virtual FString GetDefaultProjectCreationPath() = 0;
-
-	/** 
-	 * Get (or create) the unique ID used to identify this computer
-	 */
-	virtual FGuid GetMachineId() = 0;
-
-	/**
-	 * Get the Epic account ID for the user who last used the Launcher
-	 */
-	virtual FString GetEpicAccountId() = 0;
-
-	/**
-	 * Set the Epic account ID for the user who last used the Launcher
-	 */
-	virtual void SetEpicAccountId(const FString& AccountId) = 0;
-};
+};
\ No newline at end of file
diff --git a/Engine/Source/Editor/UnrealEd/Private/UnrealEd.cpp b/Engine/Source/Editor/UnrealEd/Private/UnrealEd.cpp
--- a/Engine/Source/Editor/UnrealEd/Private/UnrealEd.cpp
+++ b/Engine/Source/Editor/UnrealEd/Private/UnrealEd.cpp
@@ -17,7 +17,6 @@
 #include "EngineAnalytics.h"
 #include "IAnalyticsProvider.h"
 
-#include "DesktopPlatformModule.h"
 #include "GameProjectGenerationModule.h"
 #include "AssetEditorManager.h"
 
@@ -65,15 +64,11 @@ int32 EditorInit( IEngineLoop& EngineLoop )
 	// Let the analytics know that the editor has started
 	if ( FEngineAnalytics::IsAvailable() )
 	{
-		IDesktopPlatform* const DesktopPlatform = FDesktopPlatformModule::Get();
-		if ( DesktopPlatform )
-		{
-			TArray<FAnalyticsEventAttribute> EventAttributes;
-			EventAttributes.Add(FAnalyticsEventAttribute(TEXT("MachineID"), DesktopPlatform->GetMachineId().ToString(EGuidFormats::Digits).ToLower()));
-			EventAttributes.Add(FAnalyticsEventAttribute(TEXT("AccountID"), DesktopPlatform->GetEpicAccountId()));
+		TArray<FAnalyticsEventAttribute> EventAttributes;
+		EventAttributes.Add(FAnalyticsEventAttribute(TEXT("MachineID"), FPlatformMisc::GetMachineId().ToString(EGuidFormats::Digits).ToLower()));
+		EventAttributes.Add(FAnalyticsEventAttribute(TEXT("AccountID"), FPlatformMisc::GetEpicAccountId()));
 
-			FEngineAnalytics::GetProvider().RecordEvent(TEXT("Editor.ProgramStarted"), EventAttributes);
-		}
+		FEngineAnalytics::GetProvider().RecordEvent(TEXT("Editor.ProgramStarted"), EventAttributes);
 	}
 
 	// Initialize the misc editor
diff --git a/Engine/Source/Programs/CrashReporter/CrashReportClient/CrashReportClient.Build.cs b/Engine/Source/Programs/CrashReporter/CrashReportClient/CrashReportClient.Build.cs
--- a/Engine/Source/Programs/CrashReporter/CrashReportClient/CrashReportClient.Build.cs
+++ b/Engine/Source/Programs/CrashReporter/CrashReportClient/CrashReportClient.Build.cs
@@ -37,7 +37,6 @@ public class CrashReportClient : ModuleRules
 					"SlateCore",
 					"SlateReflector",
 					"StandaloneRenderer",
-					"DesktopPlatform",
 					"MessageLog",
 				}
 			);
diff --git a/Engine/Source/Programs/CrashReporter/CrashReportClient/Private/CrashDescription.cpp b/Engine/Source/Programs/CrashReporter/CrashReportClient/Private/CrashDescription.cpp
--- a/Engine/Source/Programs/CrashReporter/CrashReportClient/Private/CrashDescription.cpp
+++ b/Engine/Source/Programs/CrashReporter/CrashReportClient/Private/CrashDescription.cpp
@@ -8,11 +8,6 @@
 #include "CrashReportAnalytics.h"
 #include "Runtime/Analytics/Analytics/Public/Analytics.h"
 #include "Runtime/Analytics/Analytics/Public/Interfaces/IAnalyticsProvider.h"
-#if !CRASH_REPORT_UNATTENDED_ONLY
-	#include "DesktopPlatformModule.h"
-#endif // CRASH_REPORT_UNATTENDED_ONLY
-
-
 
 FCrashDescription::FCrashDescription() :
 	BuiltFromCL( -1 ),
@@ -182,17 +177,10 @@ void FCrashDescription::InitializeEngineVersion()
 
 void FCrashDescription::InitializeIDs()
 {
-#if !CRASH_REPORT_UNATTENDED_ONLY
-	IDesktopPlatform* DesktopPlatform = FDesktopPlatformModule::Get();
-
-	MachineId = DesktopPlatform->GetMachineId().ToString( EGuidFormats::Digits );
+	MachineId = FPlatformMisc::GetMachineId().ToString( EGuidFormats::Digits );
 
 	// The Epic ID can be looked up from this ID.
-	EpicAccountId = DesktopPlatform->GetEpicAccountId();
-#else
-	MachineId = FGuid::NewGuid().ToString(EGuidFormats::Digits);
-	EpicAccountId = TEXT("BogusUnattendedAccount");
-#endif // CRASH_REPORT_UNATTENDED_ONLY
+	EpicAccountId = FPlatformMisc::GetEpicAccountId();
 
 	// Get the user name only for non-UE4 releases.
 	if( !FRocketSupport::IsRocket() )
diff --git a/Engine/Source/Runtime/Core/Private/GenericPlatform/GenericPlatformMisc.cpp b/Engine/Source/Runtime/Core/Private/GenericPlatform/GenericPlatformMisc.cpp
--- a/Engine/Source/Runtime/Core/Private/GenericPlatform/GenericPlatformMisc.cpp
+++ b/Engine/Source/Runtime/Core/Private/GenericPlatform/GenericPlatformMisc.cpp
@@ -782,3 +782,37 @@ bool FGenericPlatformMisc::IsRunningOnBattery()
 {
 	return false;
 }
+
+FGuid FGenericPlatformMisc::GetMachineId()
+{
+	FGuid MachineId;
+	FString MachineIdStr;
+
+	// Check to see if we already have a valid machine ID to use
+	if( !GetStoredValue( TEXT( "Epic Games" ), TEXT( "Unreal Engine/Identifiers" ), TEXT( "MachineId" ), MachineIdStr ) || !FGuid::Parse( MachineIdStr, MachineId ) )
+	{
+		// No valid machine ID, generate and save a new one
+		MachineId = FGuid::NewGuid();
+		MachineIdStr = MachineId.ToString( EGuidFormats::Digits );
+
+		if( !SetStoredValue( TEXT( "Epic Games" ), TEXT( "Unreal Engine/Identifiers" ), TEXT( "MachineId" ), MachineIdStr ) )
+		{
+			// Failed to persist the machine ID - reset it to zero to avoid returning a transient value
+			MachineId = FGuid();
+		}
+	}
+
+	return MachineId;
+}
+
+FString FGenericPlatformMisc::GetEpicAccountId()
+{
+	FString AccountId;
+	GetStoredValue( TEXT( "Epic Games" ), TEXT( "Unreal Engine/Identifiers" ), TEXT( "AccountId" ), AccountId );
+	return AccountId;
+}
+
+void FGenericPlatformMisc::SetEpicAccountId( const FString& AccountId )
+{
+	SetStoredValue( TEXT( "Epic Games" ), TEXT( "Unreal Engine/Identifiers" ), TEXT( "AccountId" ), AccountId );
+}
\ No newline at end of file
diff --git a/Engine/Source/Runtime/Core/Public/GenericPlatform/GenericPlatformMisc.h b/Engine/Source/Runtime/Core/Public/GenericPlatform/GenericPlatformMisc.h
--- a/Engine/Source/Runtime/Core/Public/GenericPlatform/GenericPlatformMisc.h
+++ b/Engine/Source/Runtime/Core/Public/GenericPlatform/GenericPlatformMisc.h
@@ -811,6 +811,21 @@ public:
 	 */
 	static bool IsRunningOnBattery();
 
+	/**
+	 * Get (or create) the unique ID used to identify this computer
+	 */
+	static FGuid GetMachineId();
+
+	/**
+	 * Get the Epic account ID for the user who last used the Launcher
+	 */
+	static FString GetEpicAccountId();
+
+	/**
+	 * Set the Epic account ID for the user who last used the Launcher
+	 */
+	static void SetEpicAccountId( const FString& AccountId );
+
 #if !UE_BUILD_SHIPPING
 protected:
 	/** Whether the user should be prompted to allow for a remote debugger to be attached */
