Move functionality to invoke UBT into DesktopPlatform.

diff --git a/Engine/Source/Developer/DesktopPlatform/Private/DesktopPlatformBase.cpp b/Engine/Source/Developer/DesktopPlatform/Private/DesktopPlatformBase.cpp
--- a/Engine/Source/Developer/DesktopPlatform/Private/DesktopPlatformBase.cpp
+++ b/Engine/Source/Developer/DesktopPlatform/Private/DesktopPlatformBase.cpp
@@ -434,6 +434,133 @@ bool FDesktopPlatformBase::GenerateProjectFiles(const FString& RootDir, const FS
 	return bRes;
 }
 
+bool FDesktopPlatformBase::IsUnrealBuildToolAvailable()
+{
+	// If using Rocket and the Rocket unreal build tool executable exists, then UBT is available
+	if (FApp::IsEngineInstalled())
+	{
+		return FPaths::FileExists(GetUnrealBuildToolExecutableFilename());
+	}
+	else
+	{
+		// If not using Rocket, check to make sure UBT can be built, since it is an intermediate.
+		// We are simply checking for the existence of source code files, which is a heuristic to determine if UBT can be built
+		TArray<FString> Filenames;
+		const FString UBTSourcePath = GetUnrealBuildToolSourceCodePath();
+		const FString SearchPattern = TEXT("*.cs");
+		const bool bFiles = true;
+		const bool bDirectories = false;
+		const bool bClearFileNames = false;
+		IFileManager::Get().FindFilesRecursive(Filenames, *UBTSourcePath, *SearchPattern, bFiles, bDirectories, bClearFileNames);
+
+		return Filenames.Num() > 0;
+	}
+}
+
+bool FDesktopPlatformBase::InvokeUnrealBuildToolSync(const FString& InCmdLineParams, FOutputDevice &Ar, bool bSkipBuildUBT, int32& OutReturnCode, FString& OutProcOutput)
+{
+	// Setup output redirection pipes, so that we can harvest compiler output and display it ourselves
+#if PLATFORM_LINUX
+	int pipefd[2];
+	pipe(pipefd);
+	void* PipeRead = &pipefd[0];
+	void* PipeWrite = &pipefd[1];
+#else
+	void* PipeRead = NULL;
+	void* PipeWrite = NULL;
+#endif
+
+	verify(FPlatformProcess::CreatePipe(PipeRead, PipeWrite));
+
+	bool bInvoked = false;
+	FProcHandle ProcHandle = InvokeUnrealBuildToolAsync(InCmdLineParams, Ar, PipeRead, PipeWrite, bSkipBuildUBT);
+	if (ProcHandle.IsValid())
+	{
+		// rather than waiting, we must flush the read pipe or UBT will stall if it writes out a ton of text to the console.
+		while (FPlatformProcess::IsProcRunning(ProcHandle))
+		{
+			OutProcOutput += FPlatformProcess::ReadPipe(PipeRead);
+			FPlatformProcess::Sleep(0.1f);
+		}		
+		bInvoked = true;
+		bool bGotReturnCode = FPlatformProcess::GetProcReturnCode(ProcHandle, &OutReturnCode);		
+		check(bGotReturnCode);
+	}
+	else
+	{
+		bInvoked = false;
+		OutReturnCode = -1;
+		OutProcOutput = TEXT("");
+	}
+
+
+	FPlatformProcess::ClosePipe(PipeRead, PipeWrite);
+
+	return bInvoked;
+}
+
+FProcHandle FDesktopPlatformBase::InvokeUnrealBuildToolAsync(const FString& InCmdLineParams, FOutputDevice &Ar, void*& OutReadPipe, void*& OutWritePipe, bool bSkipBuildUBT)
+{
+#if PLATFORM_DESKTOP && !IS_MONOLITHIC
+	FString CmdLineParams = InCmdLineParams;
+
+	if (FRocketSupport::IsRocket())
+	{
+		CmdLineParams += TEXT(" -rocket");
+	}
+
+	// UnrealBuildTool is currently always located in the Binaries/DotNET folder
+	FString ExecutableFileName = GetUnrealBuildToolExecutableFilename();
+
+	// Rocket never builds UBT, UnrealBuildTool should already exist
+	bool bSkipBuild = FApp::IsEngineInstalled() || bSkipBuildUBT;
+	if (!bSkipBuild)
+	{
+		// When not using rocket, we should attempt to build UBT to make sure it is up to date
+		// Only do this if we have not already successfully done it once during this session.
+		static bool bSuccessfullyBuiltUBTOnce = false;
+		if (!bSuccessfullyBuiltUBTOnce)
+		{
+			Ar.Log(TEXT("Building UnrealBuildTool..."));
+			if (BuildUnrealBuildTool(Ar))
+			{
+				bSuccessfullyBuiltUBTOnce = true;
+			}
+			else
+			{
+				// Failed to build UBT
+				Ar.Log(TEXT("Failed to build UnrealBuildTool."));
+				return FProcHandle();
+			}
+		}
+	}
+
+	Ar.Logf(TEXT("Launching UnrealBuildTool... [%s %s]"), *ExecutableFileName, *CmdLineParams);
+
+#if PLATFORM_MAC
+	// On Mac we launch UBT with Mono
+	FString ScriptPath = FPaths::ConvertRelativePathToFull(FPaths::EngineDir() / TEXT("Build/BatchFiles/Mac/RunMono.sh"));
+	CmdLineParams = FString::Printf(TEXT("\"%s\" \"%s\" %s"), *ScriptPath, *ExecutableFileName, *CmdLineParams);
+	ExecutableFileName = TEXT("/bin/sh");
+#endif
+
+	// Run UnrealBuildTool
+	const bool bLaunchDetached = false;
+	const bool bLaunchHidden = true;
+	const bool bLaunchReallyHidden = bLaunchHidden;
+
+	FProcHandle ProcHandle = FPlatformProcess::CreateProc(*ExecutableFileName, *CmdLineParams, bLaunchDetached, bLaunchHidden, bLaunchReallyHidden, NULL, 0, NULL, OutWritePipe);
+	if (!ProcHandle.IsValid())
+	{
+		Ar.Logf(TEXT("Failed to launch Unreal Build Tool. (%s)"), *ExecutableFileName);
+	}
+
+	return ProcHandle;
+#else
+	return FProcHandle();
+#endif // PLATFORM_DESKTOP && !IS_MONOLITHIC
+}
+
 bool FDesktopPlatformBase::GetSolutionPath(FString& OutSolutionPath)
 {
 	// Get the platform-specific suffix for solution files
@@ -739,4 +866,108 @@ bool FDesktopPlatformBase::EnumerateProjectsKnownByEngine(const FString &Identif
 	return true;
 }
 
+FString FDesktopPlatformBase::GetUnrealBuildToolExecutableFilename() const
+{
+	const FString UBTPath = FString::Printf(TEXT("%sBinaries/DotNET/"), *FPaths::EngineDir());
+	const FString UBTExe = TEXT("UnrealBuildTool.exe");
+
+	FString ExecutableFileName = UBTPath / UBTExe;
+	ExecutableFileName = FPaths::ConvertRelativePathToFull(ExecutableFileName);
+	return ExecutableFileName;
+}
+
+FString FDesktopPlatformBase::GetUnrealBuildToolSourceCodePath() const
+{
+	return FPaths::Combine(*FPaths::EngineDir(), TEXT("Source"), TEXT("Programs"), TEXT("UnrealBuildTool"));
+}
+
+bool FDesktopPlatformBase::BuildUnrealBuildTool(FOutputDevice &Ar)
+{
+#if !IS_MONOLITHIC
+	if (FApp::IsEngineInstalled())
+	{
+		// We may not build UBT in rocket
+		return false;
+	}
+
+	FString CompilerExecutableFilename;
+	FString CmdLineParams;
+#if PLATFORM_WINDOWS
+	// To build UBT for windows, we must assemble a batch file that first registers the environment variable necessary to run msbuild then run it
+	// This can not be done in a single invocation of CMD.exe because the environment variables do not transfer between subsequent commands when using the "&" syntax
+	// devenv.exe can be used to build as well but it takes several seconds to start up so it is not desirable
+
+	// First determine the appropriate vcvars batch file to launch
+	FString VCVarsBat;
+
+#if _MSC_VER >= 1800
+	FPlatformMisc::GetVSComnTools(12, VCVarsBat);
+#else
+	FPlatformMisc::GetVSComnTools(11, VCVarsBat);
+#endif
+
+	VCVarsBat = FPaths::Combine(*VCVarsBat, L"../../VC/bin/x86_amd64/vcvarsx86_amd64.bat");
+
+	// Check to make sure we found one.
+	if (VCVarsBat.IsEmpty() || !FPaths::FileExists(VCVarsBat))
+	{
+		// VCVars doesn't exist, we can not build UBT
+		CompilerExecutableFilename = TEXT("");
+	}
+	else
+	{
+		// Now make a batch file in the intermediate directory to invoke the vcvars batch then msbuild
+		FString BuildBatchFile = FPaths::EngineIntermediateDir() / TEXT("Build") / TEXT("UnrealBuildTool") / TEXT("BuildUBT.bat");
+		BuildBatchFile.ReplaceInline(TEXT("/"), TEXT("\\"));
+
+		const FString CsProjLocation = FPaths::ConvertRelativePathToFull(GetUnrealBuildToolSourceCodePath()) / TEXT("UnrealBuildTool.csproj");
+
+		FString BatchFileContents;
+		BatchFileContents = FString::Printf(TEXT("call \"%s\"") LINE_TERMINATOR, *VCVarsBat);
+		BatchFileContents += FString::Printf(TEXT("msbuild /nologo /verbosity:quiet %s /property:Configuration=Development /property:Platform=AnyCPU"), *CsProjLocation);
+		FFileHelper::SaveStringToFile(BatchFileContents, *BuildBatchFile);
+
+		{
+			TCHAR CmdExePath[MAX_PATH];
+			FPlatformMisc::GetEnvironmentVariable(TEXT("ComSpec"), CmdExePath, ARRAY_COUNT(CmdExePath));
+			CompilerExecutableFilename = CmdExePath;
+		}
+
+		CmdLineParams = TEXT("/c ");
+		CmdLineParams += BuildBatchFile;
+	}
+#elif PLATFORM_MAC
+	const FString CsProjLocation = FPaths::ConvertRelativePathToFull(GetUnrealBuildToolSourceCodePath()) / TEXT("UnrealBuildTool_Mono.csproj");
+	FString ScriptPath = FPaths::ConvertRelativePathToFull(FPaths::EngineDir() / TEXT("Build/BatchFiles/Mac/RunXBuild.sh"));
+	CompilerExecutableFilename = TEXT("/bin/sh");
+	CmdLineParams = FString::Printf(TEXT("\"%s\" /property:Configuration=Development %s"), *ScriptPath, *CsProjLocation);
+#elif PLATFORM_LINUX
+	printf("ModuleManager.cpp: TODO: Linux BuildUBT");
+#else
+	Ar.Log(TEXT("Unknown platform, unable to build UnrealBuildTool."));
+#endif
+
+	// If a compiler executable was provided, try to build now
+	if (!CompilerExecutableFilename.IsEmpty())
+	{
+		const bool bLaunchDetached = false;
+		const bool bLaunchHidden = true;
+		const bool bLaunchReallyHidden = bLaunchHidden;
+		FProcHandle ProcHandle = FPlatformProcess::CreateProc(*CompilerExecutableFilename, *CmdLineParams, bLaunchDetached, bLaunchHidden, bLaunchReallyHidden, NULL, 0, NULL, NULL);
+		if (ProcHandle.IsValid())
+		{
+			FPlatformProcess::WaitForProc(ProcHandle);
+			ProcHandle.Close();
+		}
+
+		// If the executable appeared where we expect it, then we were successful
+		return FPaths::FileExists(GetUnrealBuildToolExecutableFilename());
+	}
+	else
+#endif // !IS_MONOLITHIC
+	{
+		return false;
+	}
+}
+
 #undef LOCTEXT_NAMESPACE
diff --git a/Engine/Source/Developer/DesktopPlatform/Private/DesktopPlatformBase.h b/Engine/Source/Developer/DesktopPlatform/Private/DesktopPlatformBase.h
--- a/Engine/Source/Developer/DesktopPlatform/Private/DesktopPlatformBase.h
+++ b/Engine/Source/Developer/DesktopPlatform/Private/DesktopPlatformBase.h
@@ -38,6 +38,9 @@ public:
 	virtual bool CleanGameProject(const FString& ProjectDir, FString& OutFailPath, FFeedbackContext* Warn) override;
 	virtual bool CompileGameProject(const FString& RootDir, const FString& ProjectFileName, FFeedbackContext* Warn) override;
 	virtual bool GenerateProjectFiles(const FString& RootDir, const FString& ProjectFileName, FFeedbackContext* Warn) override;
+	virtual bool IsUnrealBuildToolAvailable() override;
+	virtual bool InvokeUnrealBuildToolSync(const FString& InCmdLineParams, FOutputDevice &Ar, bool bSkipBuildUBT, int32& OutReturnCode, FString& OutProcOutput) override;
+	virtual FProcHandle InvokeUnrealBuildToolAsync(const FString& InCmdLineParams, FOutputDevice &Ar, void*& OutReadPipe, void*& OutWritePipe, bool bSkipBuildUBT = false) override;
 
 	virtual bool GetSolutionPath(FString& OutSolutionPath) override;
 
@@ -60,4 +63,8 @@ private:
 	const FUProjectDictionary &GetCachedProjectDictionary(const FString& RootDir);
 
 	void GetProjectBuildProducts(const FString& ProjectFileName, TArray<FString> &OutFileNames, TArray<FString> &OutDirectoryNames);
+
+	FString GetUnrealBuildToolExecutableFilename() const;
+	FString GetUnrealBuildToolSourceCodePath() const;
+	bool BuildUnrealBuildTool(FOutputDevice &Ar);
 };
diff --git a/Engine/Source/Developer/DesktopPlatform/Public/IDesktopPlatform.h b/Engine/Source/Developer/DesktopPlatform/Public/IDesktopPlatform.h
--- a/Engine/Source/Developer/DesktopPlatform/Public/IDesktopPlatform.h
+++ b/Engine/Source/Developer/DesktopPlatform/Public/IDesktopPlatform.h
@@ -329,6 +329,27 @@ public:
 	*/
 	virtual bool GenerateProjectFiles(const FString& RootDir, const FString& ProjectFileName, FFeedbackContext* Warn) = 0;
 
+	/**
+	* Determines whether UnrealBuildTool is available
+	*
+	* @return true if UnrealBuildTool is available
+	*/
+	virtual bool IsUnrealBuildToolAvailable() = 0;
+
+	/**
+	 * Invokes UnrealBuildTool with the given arguments
+	 *
+	 * @return true if tool was invoked properly
+	 */
+	virtual bool InvokeUnrealBuildToolSync(const FString& InCmdLineParams, FOutputDevice &Ar, bool bSkipBuildUBT, int32& OutReturnCode, FString& OutProcOutput) = 0;
+
+	/** 
+	 * Launches UnrealBuildTool with the specified command line parameters 
+	 * 
+	 * @return process handle to the new UBT process
+	 */
+	virtual FProcHandle InvokeUnrealBuildToolAsync(const FString& InCmdLineParams, FOutputDevice &Ar, void*& OutReadPipe, void*& OutWritePipe, bool bSkipBuildUBT = false) = 0;
+
 	/**
 	* Runs UnrealBuildTool with the given arguments.
 	*
diff --git a/Engine/Source/Developer/HotReload/HotReload.Build.cs b/Engine/Source/Developer/HotReload/HotReload.Build.cs
--- a/Engine/Source/Developer/HotReload/HotReload.Build.cs
+++ b/Engine/Source/Developer/HotReload/HotReload.Build.cs
@@ -20,6 +20,7 @@ public class HotReload : ModuleRules
 					"Core",
 					"CoreUObject",
 					"DirectoryWatcher",
+					"DesktopPlatform"
 					// ... add other public dependencies that you statically link with here ...
 				}
 				);
diff --git a/Engine/Source/Developer/HotReload/Private/HotReload.cpp b/Engine/Source/Developer/HotReload/Private/HotReload.cpp
--- a/Engine/Source/Developer/HotReload/Private/HotReload.cpp
+++ b/Engine/Source/Developer/HotReload/Private/HotReload.cpp
@@ -5,6 +5,7 @@
 #include "Runtime/Analytics/Analytics/Public/Interfaces/IAnalyticsProvider.h"
 #include "EngineAnalytics.h"
 #include "ScopedTimers.h"
+#include "DesktopPlatformModule.h"
 
 DEFINE_LOG_CATEGORY(LogHotReload);
 
@@ -1135,7 +1136,7 @@ bool FHotReloadModule::InvokeUnrealBuildToolForCompile(const FString& InCmdLineP
 	verify(FPlatformProcess::CreatePipe(PipeRead, PipeWrite));
 	ModuleCompileReadPipeText = TEXT("");
 
-	FProcHandle ProcHandle = FUBTInvoker::InvokeUnrealBuildToolAsync(InCmdLineParams, Ar, PipeRead, PipeWrite);
+	FProcHandle ProcHandle = FDesktopPlatformModule::Get()->InvokeUnrealBuildToolAsync(InCmdLineParams, Ar, PipeRead, PipeWrite);
 
 	// We no longer need the Write pipe so close it.
 	// We DO need the Read pipe however...
diff --git a/Engine/Source/Developer/TargetPlatform/Private/TargetPlatformManagerModule.cpp b/Engine/Source/Developer/TargetPlatform/Private/TargetPlatformManagerModule.cpp
--- a/Engine/Source/Developer/TargetPlatform/Private/TargetPlatformManagerModule.cpp
+++ b/Engine/Source/Developer/TargetPlatform/Private/TargetPlatformManagerModule.cpp
@@ -6,6 +6,7 @@
 
 #include "TargetPlatformPrivatePCH.h"
 #include "PlatformInfo.h"
+#include "DesktopPlatformModule.h"
 
 DEFINE_LOG_CATEGORY_STATIC(LogTargetPlatformManager, Log, All);
 
@@ -40,7 +41,7 @@ public:
 			FString UBTParams(TEXT("-autosdkonly"));
 			int32 UBTReturnCode = -1;
 			FString UBTOutput;
-			if (!FUBTInvoker::InvokeUnrealBuildToolSync(UBTParams, *GLog, true, UBTReturnCode, UBTOutput))
+			if (!FDesktopPlatformModule::Get()->InvokeUnrealBuildToolSync(UBTParams, *GLog, true, UBTReturnCode, UBTOutput))
 			{
 				UE_LOG(LogTargetPlatformManager, Fatal, TEXT("Failed to run UBT to check SDK status!"));
 			}
diff --git a/Engine/Source/Editor/GameProjectGeneration/Private/SNewProjectWizard.cpp b/Engine/Source/Editor/GameProjectGeneration/Private/SNewProjectWizard.cpp
--- a/Engine/Source/Editor/GameProjectGeneration/Private/SNewProjectWizard.cpp
+++ b/Engine/Source/Editor/GameProjectGeneration/Private/SNewProjectWizard.cpp
@@ -1123,7 +1123,7 @@ void SNewProjectWizard::UpdateProjectFileValidity( )
 					bLastGlobalValidityCheckSuccessful = false;
 					LastGlobalValidityErrorText = FText::Format( LOCTEXT("NoCompilerFound", "No compiler was found. In order to use a C++ template, you must first install {0}."), FSourceCodeNavigation::GetSuggestedSourceCodeIDE() );
 				}
-				else if ( !FModuleManager::Get().IsUnrealBuildToolAvailable() )
+				else if ( !FDesktopPlatformModule::Get()->IsUnrealBuildToolAvailable() )
 				{
 					bLastGlobalValidityCheckSuccessful = false;
 					LastGlobalValidityErrorText = LOCTEXT("UBTNotFound", "Engine source code was not found. In order to use a C++ template, you must have engine source code in Engine/Source.");
diff --git a/Engine/Source/Runtime/Core/Private/Modules/ModuleManager.cpp b/Engine/Source/Runtime/Core/Private/Modules/ModuleManager.cpp
--- a/Engine/Source/Runtime/Core/Private/Modules/ModuleManager.cpp
+++ b/Engine/Source/Runtime/Core/Private/Modules/ModuleManager.cpp
@@ -967,235 +967,3 @@ bool FModuleManager::DoesLoadedModuleHaveUObjects( const FName ModuleName )
 
 	return false;
 }
-
-bool FModuleManager::IsUnrealBuildToolAvailable()
-{
-	// If using Rocket and the Rocket unreal build tool executable exists, then UBT is available
-	if (FApp::IsEngineInstalled())
-	{
-		return FPaths::FileExists(FUBTInvoker::GetUnrealBuildToolExecutableFilename());
-	}
-	else
-	{
-		// If not using Rocket, check to make sure UBT can be built, since it is an intermediate.
-		// We are simply checking for the existence of source code files, which is a heuristic to determine if UBT can be built
-		TArray<FString> Filenames;
-		const FString UBTSourcePath = FUBTInvoker::GetUnrealBuildToolSourceCodePath();
-		const FString SearchPattern = TEXT("*.cs");
-		const bool bFiles = true;
-		const bool bDirectories = false;
-		const bool bClearFileNames = false;
-		IFileManager::Get().FindFilesRecursive(Filenames, *UBTSourcePath, *SearchPattern, bFiles, bDirectories, bClearFileNames);
-
-		return Filenames.Num() > 0;
-	}
-}
-
-FString FUBTInvoker::GetUnrealBuildToolExecutableFilename()
-{
-	const FString UBTPath = FString::Printf(TEXT("%sBinaries/DotNET/"), *FPaths::EngineDir());
-	const FString UBTExe = TEXT("UnrealBuildTool.exe");
-
-	FString ExecutableFileName = UBTPath / UBTExe;
-	ExecutableFileName = FPaths::ConvertRelativePathToFull(ExecutableFileName);
-	return ExecutableFileName;
-}
-
-
-bool FUBTInvoker::BuildUnrealBuildTool(FOutputDevice &Ar)
-{
-#if !IS_MONOLITHIC
-	if (FApp::IsEngineInstalled())
-	{
-		// We may not build UBT in rocket
-		return false;
-	}
-
-	FString CompilerExecutableFilename;
-	FString CmdLineParams;
-#if PLATFORM_WINDOWS
-	// To build UBT for windows, we must assemble a batch file that first registers the environment variable necessary to run msbuild then run it
-	// This can not be done in a single invocation of CMD.exe because the environment variables do not transfer between subsequent commands when using the "&" syntax
-	// devenv.exe can be used to build as well but it takes several seconds to start up so it is not desirable
-
-	// First determine the appropriate vcvars batch file to launch
-	FString VCVarsBat;
-
-#if _MSC_VER >= 1800
-	FPlatformMisc::GetVSComnTools(12, VCVarsBat);
-#else
-	FPlatformMisc::GetVSComnTools(11, VCVarsBat);
-#endif
-
-	VCVarsBat = FPaths::Combine(*VCVarsBat, L"../../VC/bin/x86_amd64/vcvarsx86_amd64.bat");
-
-	// Check to make sure we found one.
-	if (VCVarsBat.IsEmpty() || !FPaths::FileExists(VCVarsBat))
-	{
-		// VCVars doesn't exist, we can not build UBT
-		CompilerExecutableFilename = TEXT("");
-	}
-	else
-	{
-		// Now make a batch file in the intermediate directory to invoke the vcvars batch then msbuild
-		FString BuildBatchFile = FPaths::EngineIntermediateDir() / TEXT("Build") / TEXT("UnrealBuildTool") / TEXT("BuildUBT.bat");
-		BuildBatchFile.ReplaceInline(TEXT("/"), TEXT("\\"));
-
-		const FString CsProjLocation = FPaths::ConvertRelativePathToFull(GetUnrealBuildToolSourceCodePath()) / TEXT("UnrealBuildTool.csproj");
-
-		FString BatchFileContents;
-		BatchFileContents = FString::Printf(TEXT("call \"%s\"") LINE_TERMINATOR, *VCVarsBat);
-		BatchFileContents += FString::Printf(TEXT("msbuild /nologo /verbosity:quiet %s /property:Configuration=Development /property:Platform=AnyCPU"), *CsProjLocation);
-		FFileHelper::SaveStringToFile(BatchFileContents, *BuildBatchFile);
-
-		{
-			TCHAR CmdExePath[MAX_PATH];
-			FPlatformMisc::GetEnvironmentVariable(TEXT("ComSpec"), CmdExePath, ARRAY_COUNT(CmdExePath));
-			CompilerExecutableFilename = CmdExePath;
-		}
-
-		CmdLineParams = TEXT("/c ");
-		CmdLineParams += BuildBatchFile;
-	}
-#elif PLATFORM_MAC
-	const FString CsProjLocation = FPaths::ConvertRelativePathToFull(GetUnrealBuildToolSourceCodePath()) / TEXT("UnrealBuildTool_Mono.csproj");
-	FString ScriptPath = FPaths::ConvertRelativePathToFull(FPaths::EngineDir() / TEXT("Build/BatchFiles/Mac/RunXBuild.sh"));
-	CompilerExecutableFilename = TEXT("/bin/sh");
-	CmdLineParams = FString::Printf(TEXT("\"%s\" /property:Configuration=Development %s"), *ScriptPath, *CsProjLocation);
-#elif PLATFORM_LINUX
-	printf("ModuleManager.cpp: TODO: Linux BuildUBT");
-#else
-	Ar.Log(TEXT("Unknown platform, unable to build UnrealBuildTool."));
-#endif
-
-	// If a compiler executable was provided, try to build now
-	if (!CompilerExecutableFilename.IsEmpty())
-	{
-		const bool bLaunchDetached = false;
-		const bool bLaunchHidden = true;
-		const bool bLaunchReallyHidden = bLaunchHidden;
-		FProcHandle ProcHandle = FPlatformProcess::CreateProc(*CompilerExecutableFilename, *CmdLineParams, bLaunchDetached, bLaunchHidden, bLaunchReallyHidden, NULL, 0, NULL, NULL);
-		if (ProcHandle.IsValid())
-		{
-			FPlatformProcess::WaitForProc(ProcHandle);
-			ProcHandle.Close();
-		}
-
-		// If the executable appeared where we expect it, then we were successful
-		return FPaths::FileExists(GetUnrealBuildToolExecutableFilename());
-	}
-	else
-#endif // !IS_MONOLITHIC
-	{
-		return false;
-	}
-}
-
-bool FUBTInvoker::InvokeUnrealBuildToolSync(const FString& InCmdLineParams, FOutputDevice &Ar, bool bSkipBuildUBT, int32& OutReturnCode, FString& OutProcOutput)
-{
-	// Setup output redirection pipes, so that we can harvest compiler output and display it ourselves
-#if PLATFORM_LINUX
-	int pipefd[2];
-	pipe(pipefd);
-	void* PipeRead = &pipefd[0];
-	void* PipeWrite = &pipefd[1];
-#else
-	void* PipeRead = NULL;
-	void* PipeWrite = NULL;
-#endif
-
-	verify(FPlatformProcess::CreatePipe(PipeRead, PipeWrite));
-
-	bool bInvoked = false;
-	FProcHandle ProcHandle = InvokeUnrealBuildToolAsync(InCmdLineParams, Ar, PipeRead, PipeWrite, bSkipBuildUBT);
-	if (ProcHandle.IsValid())
-	{
-		// rather than waiting, we must flush the read pipe or UBT will stall if it writes out a ton of text to the console.
-		while (FPlatformProcess::IsProcRunning(ProcHandle))
-		{
-			OutProcOutput += FPlatformProcess::ReadPipe(PipeRead);
-			FPlatformProcess::Sleep(0.1f);
-		}		
-		bInvoked = true;
-		bool bGotReturnCode = FPlatformProcess::GetProcReturnCode(ProcHandle, &OutReturnCode);		
-		check(bGotReturnCode);
-	}
-	else
-	{
-		bInvoked = false;
-		OutReturnCode = -1;
-		OutProcOutput = TEXT("");
-	}
-
-
-	FPlatformProcess::ClosePipe(PipeRead, PipeWrite);
-
-	return bInvoked;
-}
-
-FProcHandle FUBTInvoker::InvokeUnrealBuildToolAsync(const FString& InCmdLineParams, FOutputDevice &Ar, void*& OutReadPipe, void*& OutWritePipe, bool bSkipBuildUBT)
-{
-#if PLATFORM_DESKTOP && !IS_MONOLITHIC
-	FString CmdLineParams = InCmdLineParams;
-
-	if (FRocketSupport::IsRocket())
-	{
-		CmdLineParams += TEXT(" -rocket");
-	}
-
-	// UnrealBuildTool is currently always located in the Binaries/DotNET folder
-	FString ExecutableFileName = GetUnrealBuildToolExecutableFilename();
-
-	// Rocket never builds UBT, UnrealBuildTool should already exist
-	bool bSkipBuild = FApp::IsEngineInstalled() || bSkipBuildUBT;
-	if (!bSkipBuild)
-	{
-		// When not using rocket, we should attempt to build UBT to make sure it is up to date
-		// Only do this if we have not already successfully done it once during this session.
-		static bool bSuccessfullyBuiltUBTOnce = false;
-		if (!bSuccessfullyBuiltUBTOnce)
-		{
-			Ar.Log(TEXT("Building UnrealBuildTool..."));
-			if (BuildUnrealBuildTool(Ar))
-			{
-				bSuccessfullyBuiltUBTOnce = true;
-			}
-			else
-			{
-				// Failed to build UBT
-				Ar.Log(TEXT("Failed to build UnrealBuildTool."));
-				return FProcHandle();
-			}
-		}
-	}
-
-	Ar.Logf(TEXT("Launching UnrealBuildTool... [%s %s]"), *ExecutableFileName, *CmdLineParams);
-
-#if PLATFORM_MAC
-	// On Mac we launch UBT with Mono
-	FString ScriptPath = FPaths::ConvertRelativePathToFull(FPaths::EngineDir() / TEXT("Build/BatchFiles/Mac/RunMono.sh"));
-	CmdLineParams = FString::Printf(TEXT("\"%s\" \"%s\" %s"), *ScriptPath, *ExecutableFileName, *CmdLineParams);
-	ExecutableFileName = TEXT("/bin/sh");
-#endif
-
-	// Run UnrealBuildTool
-	const bool bLaunchDetached = false;
-	const bool bLaunchHidden = true;
-	const bool bLaunchReallyHidden = bLaunchHidden;
-
-	FProcHandle ProcHandle = FPlatformProcess::CreateProc(*ExecutableFileName, *CmdLineParams, bLaunchDetached, bLaunchHidden, bLaunchReallyHidden, NULL, 0, NULL, OutWritePipe);
-	if (!ProcHandle.IsValid())
-	{
-		Ar.Logf(TEXT("Failed to launch Unreal Build Tool. (%s)"), *ExecutableFileName);
-	}
-
-	return ProcHandle;
-#else
-	return FProcHandle();
-#endif // PLATFORM_DESKTOP && !IS_MONOLITHIC
-}
-
-FString FUBTInvoker::GetUnrealBuildToolSourceCodePath()
-{
-	return FPaths::Combine(*FPaths::EngineDir(), TEXT("Source"), TEXT("Programs"), TEXT("UnrealBuildTool"));
-}
diff --git a/Engine/Source/Runtime/Core/Public/Modules/ModuleManager.h b/Engine/Source/Runtime/Core/Public/Modules/ModuleManager.h
--- a/Engine/Source/Runtime/Core/Public/Modules/ModuleManager.h
+++ b/Engine/Source/Runtime/Core/Public/Modules/ModuleManager.h
@@ -77,30 +77,6 @@ struct FModuleStatus
 	bool bIsGameModule;
 };
 
-/**
- * Class for managing Sync/Async calls to UBT with the option to rebuild UBT.
- * Separated from ModuleManager in prep for merging this into the UBT calling code in DesktopPlatform.
- */
-class CORE_API FUBTInvoker
-{
-public:	
-
-	/** Returns the filename for UBT including the path */
-	static FString GetUnrealBuildToolExecutableFilename();
-
-	/** Returns true if tool was invoked properly */
-	static bool InvokeUnrealBuildToolSync(const FString& InCmdLineParams, FOutputDevice &Ar, bool bSkipBuildUBT, int32& OutReturnCode, FString& OutProcOutput);
-
-	/** Launches UnrealBuildTool with the specified command line parameters */
-	static FProcHandle InvokeUnrealBuildToolAsync(const FString& InCmdLineParams, FOutputDevice &Ar, void*& OutReadPipe, void*& OutWritePipe, bool bSkipBuildUBT = false);
-
-	/** Builds unreal build tool using a compiler specific to the currently running platform */
-	static bool BuildUnrealBuildTool(FOutputDevice &Ar);
-
-	/** Returns the path to the unreal build tool source code */
-	static FString GetUnrealBuildToolSourceCodePath();	
-};
-
 /**
  * Implements the module manager.
  *
@@ -442,13 +418,6 @@ public:
 		return IsPackageLoaded;
 	}
 
-public:
-
-	/**
-	* @return true if the UBT executable exists (in Rocket) or source code is available to compile it (in non-Rocket)
-	*/
-	bool IsUnrealBuildToolAvailable();
-
 public:
 
 	// FSelfRegisteringExec interface.
