Remove a couple of unnecessary IS_MONOLITHIC checks.

diff --git a/Engine/Source/Developer/DesktopPlatform/Private/DesktopPlatformBase.cpp b/Engine/Source/Developer/DesktopPlatform/Private/DesktopPlatformBase.cpp
--- a/Engine/Source/Developer/DesktopPlatform/Private/DesktopPlatformBase.cpp
+++ b/Engine/Source/Developer/DesktopPlatform/Private/DesktopPlatformBase.cpp
@@ -501,7 +501,6 @@ bool FDesktopPlatformBase::InvokeUnrealBuildToolSync(const FString& InCmdLinePar
 
 FProcHandle FDesktopPlatformBase::InvokeUnrealBuildToolAsync(const FString& InCmdLineParams, FOutputDevice &Ar, void*& OutReadPipe, void*& OutWritePipe, bool bSkipBuildUBT)
 {
-#if PLATFORM_DESKTOP && !IS_MONOLITHIC
 	FString CmdLineParams = InCmdLineParams;
 
 	if (FRocketSupport::IsRocket())
@@ -556,9 +555,6 @@ FProcHandle FDesktopPlatformBase::InvokeUnrealBuildToolAsync(const FString& InCm
 	}
 
 	return ProcHandle;
-#else
-	return FProcHandle();
-#endif // PLATFORM_DESKTOP && !IS_MONOLITHIC
 }
 
 bool FDesktopPlatformBase::GetSolutionPath(FString& OutSolutionPath)
@@ -883,7 +879,6 @@ FString FDesktopPlatformBase::GetUnrealBuildToolSourceCodePath() const
 
 bool FDesktopPlatformBase::BuildUnrealBuildTool(FOutputDevice &Ar)
 {
-#if !IS_MONOLITHIC
 	if (FApp::IsEngineInstalled())
 	{
 		// We may not build UBT in rocket
@@ -948,26 +943,23 @@ bool FDesktopPlatformBase::BuildUnrealBuildTool(FOutputDevice &Ar)
 #endif
 
 	// If a compiler executable was provided, try to build now
-	if (!CompilerExecutableFilename.IsEmpty())
+	if (CompilerExecutableFilename.IsEmpty())
 	{
-		const bool bLaunchDetached = false;
-		const bool bLaunchHidden = true;
-		const bool bLaunchReallyHidden = bLaunchHidden;
-		FProcHandle ProcHandle = FPlatformProcess::CreateProc(*CompilerExecutableFilename, *CmdLineParams, bLaunchDetached, bLaunchHidden, bLaunchReallyHidden, NULL, 0, NULL, NULL);
-		if (ProcHandle.IsValid())
-		{
-			FPlatformProcess::WaitForProc(ProcHandle);
-			ProcHandle.Close();
-		}
-
-		// If the executable appeared where we expect it, then we were successful
-		return FPaths::FileExists(GetUnrealBuildToolExecutableFilename());
+		return false;
 	}
-	else
-#endif // !IS_MONOLITHIC
+
+	const bool bLaunchDetached = false;
+	const bool bLaunchHidden = true;
+	const bool bLaunchReallyHidden = bLaunchHidden;
+	FProcHandle ProcHandle = FPlatformProcess::CreateProc(*CompilerExecutableFilename, *CmdLineParams, bLaunchDetached, bLaunchHidden, bLaunchReallyHidden, NULL, 0, NULL, NULL);
+	if (ProcHandle.IsValid())
 	{
-		return false;
+		FPlatformProcess::WaitForProc(ProcHandle);
+		ProcHandle.Close();
 	}
+
+	// If the executable appeared where we expect it, then we were successful
+	return FPaths::FileExists(GetUnrealBuildToolExecutableFilename());
 }
 
 #undef LOCTEXT_NAMESPACE
