Rocket: Enable HTML5 for rocket builds on Win64.

diff --git a/Engine/Source/Developer/DesktopPlatform/Private/PlatformInfo.cpp b/Engine/Source/Developer/DesktopPlatform/Private/PlatformInfo.cpp
--- a/Engine/Source/Developer/DesktopPlatform/Private/PlatformInfo.cpp
+++ b/Engine/Source/Developer/DesktopPlatform/Private/PlatformInfo.cpp
@@ -102,7 +102,7 @@ static const FPlatformInfo PlatformInfoArray[] = {
 	BuildPlatformInfo(TEXT("Android_ETC2"),				TEXT("Android_ETC2"),		LOCTEXT("Android_ETC2", "Android (ETC2)"),					EPlatformType::Game,		EPlatformFlags::CookFlavor,		FPlatformIconPaths(TEXT("Launcher/AndroidTarget/Platform_Android_ETC2_24x"), TEXT("Launcher/AndroidTarget/Platform_Android_128x")),					TEXT("-targetplatform=Android -cookflavor=ETC2"),		TEXT("Android"),	EPlatformSDKStatus::Unknown,	TEXT("Shared/Tutorials/SettingUpAndroidTutorial"),			IsAvailableOnWindows,						true),
 	BuildPlatformInfo(TEXT("Android_PVRTC"),			TEXT("Android_PVRTC"),		LOCTEXT("Android_PVRTC", "Android (PVRTC)"),				EPlatformType::Game,		EPlatformFlags::CookFlavor,		FPlatformIconPaths(TEXT("Launcher/AndroidTarget/Platform_Android_PVRTC_24x"), TEXT("Launcher/AndroidTarget/Platform_Android_128x")),				TEXT("-targetplatform=Android -cookflavor=PVRTC"),		TEXT("Android"),	EPlatformSDKStatus::Unknown,	TEXT("Shared/Tutorials/SettingUpAndroidTutorial"),			IsAvailableOnWindows,						true),
 
-	BuildPlatformInfo(TEXT("HTML5"),					TEXT("HTML5"),				LOCTEXT("HTML5", "HTML5"),									EPlatformType::Game,		EPlatformFlags::None,			FPlatformIconPaths(TEXT("Launcher/HTML5Target/Platform_HTML5_24x"), TEXT("Launcher/HTML5Target/Platform_HTML5_128x")),								TEXT(""),												TEXT("HTML5"),		EPlatformSDKStatus::Unknown,	TEXT("Platforms/HTML5/GettingStarted"),						IsAvailableOnWindows,						false),
+	BuildPlatformInfo(TEXT("HTML5"),					TEXT("HTML5"),				LOCTEXT("HTML5", "HTML5"),									EPlatformType::Game,		EPlatformFlags::None,			FPlatformIconPaths(TEXT("Launcher/HTML5Target/Platform_HTML5_24x"), TEXT("Launcher/HTML5Target/Platform_HTML5_128x")),								TEXT(""),												TEXT("HTML5"),		EPlatformSDKStatus::Unknown,	TEXT("Platforms/HTML5/GettingStarted"),						IsAvailableOnWindows ,	                    true),
 
 	BuildPlatformInfo(TEXT("PS4"),						TEXT("PS4"),				LOCTEXT("PS4", "PlayStation 4"),							EPlatformType::Game,		EPlatformFlags::None,			FPlatformIconPaths(TEXT("Launcher/PS4Target/Platform_PS4_24x"), TEXT("Launcher/PS4Target/Platform_PS4_128x")),										TEXT(""),												TEXT("PS4"),		EPlatformSDKStatus::Unknown,	TEXT("Platforms/PS4/GettingStarted"),						IsAvailableOnWindows,						false),
 
diff --git a/Engine/Source/Developer/TargetPlatform/TargetPlatform.Build.cs b/Engine/Source/Developer/TargetPlatform/TargetPlatform.Build.cs
--- a/Engine/Source/Developer/TargetPlatform/TargetPlatform.Build.cs
+++ b/Engine/Source/Developer/TargetPlatform/TargetPlatform.Build.cs
@@ -52,10 +52,7 @@ public class TargetPlatform : ModuleRules
 					DynamicallyLoadedModuleNames.Add("Android_ETC1TargetPlatform");
 					DynamicallyLoadedModuleNames.Add("Android_ETC2TargetPlatform");
 					DynamicallyLoadedModuleNames.Add("IOSTargetPlatform");
-					if (!UnrealBuildTool.UnrealBuildTool.RunningRocket() && !UnrealBuildTool.UnrealBuildTool.BuildingRocket())
-					{
-						DynamicallyLoadedModuleNames.Add("HTML5TargetPlatform");
-					}
+                    DynamicallyLoadedModuleNames.Add("HTML5TargetPlatform");
 				}
 			}
 			else if (Target.Platform == UnrealTargetPlatform.Mac)
diff --git a/Engine/Source/Programs/AutomationTool/Scripts/GUBP.Automation.cs b/Engine/Source/Programs/AutomationTool/Scripts/GUBP.Automation.cs
--- a/Engine/Source/Programs/AutomationTool/Scripts/GUBP.Automation.cs
+++ b/Engine/Source/Programs/AutomationTool/Scripts/GUBP.Automation.cs
@@ -4813,14 +4813,14 @@ public class GUBP : BuildCommand
         }
         else
         {
-            ActivePlatforms = new List<UnrealTargetPlatform>(CommandUtils.KnownTargetPlatforms);
+            ActivePlatforms     = new List<UnrealTargetPlatform>(CommandUtils.KnownTargetPlatforms);
         }
         if (bBuildRocket)
         {
             var FilteredActivePlatforms = new List<UnrealTargetPlatform>();
             foreach (var Plat in ActivePlatforms)
             {
-                if (Plat != UnrealTargetPlatform.HTML5 && 
+                if (
                     Plat != UnrealTargetPlatform.Linux &&
                     (Plat != UnrealTargetPlatform.PS4 || ParseParam("WithPS4")) &&
                     Plat != UnrealTargetPlatform.WinRT &&
diff --git a/Engine/Source/Programs/UnrealBuildTool/HTML5/UEBuildHTML5.cs b/Engine/Source/Programs/UnrealBuildTool/HTML5/UEBuildHTML5.cs
--- a/Engine/Source/Programs/UnrealBuildTool/HTML5/UEBuildHTML5.cs
+++ b/Engine/Source/Programs/UnrealBuildTool/HTML5/UEBuildHTML5.cs
@@ -110,7 +110,7 @@ namespace UnrealBuildTool
 
         public override bool CanUseXGE()
         {
-            return (GetActiveArchitecture() == "-win32" );
+            return (GetActiveArchitecture() == "-win32");
         }
 
         /**
@@ -118,12 +118,6 @@ namespace UnrealBuildTool
          */
         protected override void RegisterBuildPlatformInternal()
         {
-            [REDACTED_EMAIL]: Add platform support
-            if (UnrealBuildTool.RunningRocket() || UnrealBuildTool.BuildingRocket())
-            {
-                return;
-            }
-
             // Make sure the SDK is installed
             if ((ProjectFileGenerator.bGenerateProjectFiles == true) || (HasRequiredSDKsInstalled() == SDKStatus.Valid))
             {
@@ -388,8 +382,13 @@ namespace UnrealBuildTool
                     }
                     else if (InModule.ToString() == "NetworkFileSystem") // server 
                     {
-                        InModule.AddPublicDependencyModule("WebSockets");
-                        InModule.AddPublicDefinition("ENABLE_HTTP_FOR_NFS=1");
+                        if (UnrealBuildTool.RunningRocket() == false || 
+                            Target.Type == TargetRules.TargetType.Game || 
+                            Target.Type == TargetRules.TargetType.RocketGame )
+                        {
+                            InModule.AddPrivateDependencyModule("WebSockets");
+                            InModule.AddPublicDefinition("ENABLE_HTTP_FOR_NFS=1");
+                        }
                     }
                 }
             }
diff --git a/Engine/Source/Runtime/HTML5/HTML5JS/Private/HTML5JavaScriptFx.cpp b/Engine/Source/Runtime/HTML5/HTML5JS/Private/HTML5JavaScriptFx.cpp
new file mode 100644
--- /dev/null
+++ b/Engine/Source/Runtime/HTML5/HTML5JS/Private/HTML5JavaScriptFx.cpp
@@ -0,0 +1,5 @@
+// Copyright 1998-2014 Epic Games, Inc. All Rights Reserved.
+#include "HTML5JavaScriptFx.h"
+
+// make build systems happy. 
+static bool THISFUNCTIONDOESNOTHING(){ return true; }
\ No newline at end of file
