Use WinXP compatible functions for DesktopPlatform functionality. TTP 345688 seems to indicate failure to compile CrashReportClient due to using RegGetValue.

diff --git a/Engine/Source/Developer/DesktopPlatform/Private/Windows/DesktopPlatformWindows.cpp b/Engine/Source/Developer/DesktopPlatform/Private/Windows/DesktopPlatformWindows.cpp
--- a/Engine/Source/Developer/DesktopPlatform/Private/Windows/DesktopPlatformWindows.cpp
+++ b/Engine/Source/Developer/DesktopPlatform/Private/Windows/DesktopPlatformWindows.cpp
@@ -619,15 +619,25 @@ void FDesktopPlatformWindows::GetRequiredRegistrySettings(TIndirectArray<FRegist
 	ShellVersionKey->FindOrAddKey(L"command")->SetValue(TEXT(""), QuotedExecutableFileName + FString(TEXT(" /switchversion \"%1\"")));
 
 	// If the user has manually selected something other than our extension, we need to delete it. Explorer explicitly disables set access on that keys in that folder, but we can delete the whole thing.
-	TCHAR ProgId[128];
-	::DWORD ProgIdSize = sizeof(ProgId);
-	if(RegGetValue(HKEY_CURRENT_USER, TEXT("Software\\Microsoft\\Windows\\CurrentVersion\\Explorer\\FileExts\\.uproject\\UserChoice"), TEXT("Progid"), RRF_RT_REG_SZ, NULL, ProgId, &ProgIdSize) != S_OK)
+	const TCHAR* UserChoicePath = TEXT("Software\\Microsoft\\Windows\\CurrentVersion\\Explorer\\FileExts\\.uproject\\UserChoice");
+
+	// Figure out whether we need to delete it. If it's already set to our own ProgId, leave it alone.
+	bool bDeleteUserChoiceKey = true;
+	HKEY hUserChoiceKey;
+	if(RegOpenKeyEx(HKEY_CURRENT_USER, UserChoicePath, 0, KEY_READ, &hUserChoiceKey) == S_OK)
 	{
-		ProgId[0] = 0;
+		TCHAR ProgId[128];
+		::DWORD ProgIdSize = sizeof(ProgId);
+		::DWORD ProgIdType = 0;
+		if(RegQueryValueEx(hUserChoiceKey, TEXT("Progid"), NULL, &ProgIdType, (BYTE*)ProgId, &ProgIdSize) == ERROR_SUCCESS && ProgIdType == REG_SZ)
+		{
+			bDeleteUserChoiceKey = (FCString::Strcmp(ProgId, TEXT("Unreal.ProjectFile")) != 0);
+		}
+		RegCloseKey(hUserChoiceKey);
 	}
-	if(FCString::Strcmp(ProgId, TEXT("Unreal.ProjectFile")) != 0)
+	if(bDeleteUserChoiceKey)
 	{
-		RootedKeys.AddRawItem(new FRegistryRootedKey(HKEY_CURRENT_USER, TEXT("Software\\Microsoft\\Windows\\CurrentVersion\\Explorer\\FileExts\\.uproject\\UserChoice")));
+		RootedKeys.AddRawItem(new FRegistryRootedKey(HKEY_CURRENT_USER, UserChoicePath));
 	}
 }
 
