Hot-Reload: Make sure hot-reload is waiting for new modules only when UBT from the current branch is running

#change Added (desktop) platform abstraction for checking if UBT is running. Fully functional for Windows, WIP for Mac and Linux (should work, but can give false positives)

diff --git a/Engine/Source/Developer/DesktopPlatform/Private/Linux/DesktopPlatformLinux.cpp b/Engine/Source/Developer/DesktopPlatform/Private/Linux/DesktopPlatformLinux.cpp
--- a/Engine/Source/Developer/DesktopPlatform/Private/Linux/DesktopPlatformLinux.cpp
+++ b/Engine/Source/Developer/DesktopPlatform/Private/Linux/DesktopPlatformLinux.cpp
@@ -360,6 +360,13 @@ bool FDesktopPlatformLinux::RunUnrealBuildTool(const FText& Description, const F
 	return FFeedbackContextMarkup::PipeProcessOutput(Description, TEXT("/usr/bin/mono"), CmdLineParams, Warn, &ExitCode) && ExitCode == 0;
 }
 
+bool FDesktopPlatformLinux::IsUnrealBuildToolRunning()
+{
+	// For now assume that if mono application is running, we're running UBT
+	// @todo: we need to get the commandline for the mono process and check if UBT.exe is in there.
+	return FPlatformProcess::IsApplicationRunning(TEXT("mono"));
+}
+
 FFeedbackContext* FDesktopPlatformLinux::GetNativeFeedbackContext()
 {
 	//unimplemented();
diff --git a/Engine/Source/Developer/DesktopPlatform/Private/Linux/DesktopPlatformLinux.h b/Engine/Source/Developer/DesktopPlatform/Private/Linux/DesktopPlatformLinux.h
--- a/Engine/Source/Developer/DesktopPlatform/Private/Linux/DesktopPlatformLinux.h
+++ b/Engine/Source/Developer/DesktopPlatform/Private/Linux/DesktopPlatformLinux.h
@@ -27,6 +27,7 @@ public:
 	virtual bool OpenProject(const FString &ProjectFileName) override;
 
 	virtual bool RunUnrealBuildTool(const FText& Description, const FString& RootDir, const FString& Arguments, FFeedbackContext* Warn) override;
+	virtual bool IsUnrealBuildToolRunning() override;
 
 	virtual FFeedbackContext* GetNativeFeedbackContext() override;
 
diff --git a/Engine/Source/Developer/DesktopPlatform/Private/Mac/DesktopPlatformMac.cpp b/Engine/Source/Developer/DesktopPlatform/Private/Mac/DesktopPlatformMac.cpp
--- a/Engine/Source/Developer/DesktopPlatform/Private/Mac/DesktopPlatformMac.cpp
+++ b/Engine/Source/Developer/DesktopPlatform/Private/Mac/DesktopPlatformMac.cpp
@@ -650,6 +650,12 @@ bool FDesktopPlatformMac::RunUnrealBuildTool(const FText& Description, const FSt
 	return FFeedbackContextMarkup::PipeProcessOutput(Description, TEXT("/bin/sh"), CmdLineParams, Warn, &ExitCode) && ExitCode == 0;
 }
 
+bool FDesktopPlatformMac::IsUnrealBuildToolRunning()
+{
+	// For now assume that if mono application is running, we're running UBT
+	// @todo: we need to get the commandline for the mono process and check if UBT.exe is in there.
+	return FPlatformProcess::IsApplicationRunning(TEXT("mono"));
+}
 
 bool FDesktopPlatformMac::GetLauncherPath(FString& OutLauncherPath) const
 {
diff --git a/Engine/Source/Developer/DesktopPlatform/Private/Mac/DesktopPlatformMac.h b/Engine/Source/Developer/DesktopPlatform/Private/Mac/DesktopPlatformMac.h
--- a/Engine/Source/Developer/DesktopPlatform/Private/Mac/DesktopPlatformMac.h
+++ b/Engine/Source/Developer/DesktopPlatform/Private/Mac/DesktopPlatformMac.h
@@ -23,6 +23,7 @@ public:
 	virtual bool UpdateFileAssociations() override;
 
 	virtual bool RunUnrealBuildTool(const FText& Description, const FString& RootDir, const FString& Arguments, FFeedbackContext* Warn) override;
+	virtual bool IsUnrealBuildToolRunning() override;
 
 	virtual FFeedbackContext* GetNativeFeedbackContext() override;
 
diff --git a/Engine/Source/Developer/DesktopPlatform/Private/Windows/DesktopPlatformWindows.cpp b/Engine/Source/Developer/DesktopPlatform/Private/Windows/DesktopPlatformWindows.cpp
--- a/Engine/Source/Developer/DesktopPlatform/Private/Windows/DesktopPlatformWindows.cpp
+++ b/Engine/Source/Developer/DesktopPlatform/Private/Windows/DesktopPlatformWindows.cpp
@@ -9,6 +9,9 @@
 	#include <shellapi.h>
 	#include <shlobj.h>
 	#include <Winver.h>
+	#include <LM.h>
+	#include <tlhelp32.h>
+	#include <Psapi.h>
 #include "HideWindowsPlatformTypes.h"
 
 #pragma comment( lib, "version.lib" )
@@ -541,6 +544,34 @@ bool FDesktopPlatformWindows::RunUnrealBuildTool(const FText& Description, const
 	return FFeedbackContextMarkup::PipeProcessOutput(Description, UnrealBuildToolPath, Arguments, Warn, &ExitCode) && ExitCode == 0;
 }
 
+bool FDesktopPlatformWindows::IsUnrealBuildToolRunning()
+{
+	FString UBTPath = GetUnrealBuildToolExecutableFilename();
+	FPaths::MakePlatformFilename(UBTPath);
+
+	HANDLE SnapShot = ::CreateToolhelp32Snapshot(TH32CS_SNAPPROCESS, 0);
+	if (SnapShot != INVALID_HANDLE_VALUE)
+	{
+		PROCESSENTRY32 Entry;
+		Entry.dwSize = sizeof(PROCESSENTRY32);
+
+		if (::Process32First(SnapShot, &Entry))
+		{
+			do
+			{
+				const FString EntryFullPath = FPlatformProcess::GetApplicationName(Entry.th32ProcessID);
+				if (EntryFullPath == UBTPath)
+				{
+					::CloseHandle(SnapShot);
+					return true;
+				}
+			} while (::Process32Next(SnapShot, &Entry));
+		}
+	}
+	::CloseHandle(SnapShot);
+	return false;
+}
+
 FFeedbackContext* FDesktopPlatformWindows::GetNativeFeedbackContext()
 {
 	static FWindowsNativeFeedbackContext FeedbackContext;
diff --git a/Engine/Source/Developer/DesktopPlatform/Private/Windows/DesktopPlatformWindows.h b/Engine/Source/Developer/DesktopPlatform/Private/Windows/DesktopPlatformWindows.h
--- a/Engine/Source/Developer/DesktopPlatform/Private/Windows/DesktopPlatformWindows.h
+++ b/Engine/Source/Developer/DesktopPlatform/Private/Windows/DesktopPlatformWindows.h
@@ -28,6 +28,7 @@ public:
 	virtual bool OpenProject(const FString &ProjectFileName) override;
 
 	virtual bool RunUnrealBuildTool(const FText& Description, const FString& RootDir, const FString& Arguments, FFeedbackContext* Warn) override;
+	virtual bool IsUnrealBuildToolRunning() override;
 
 	virtual FFeedbackContext* GetNativeFeedbackContext() override;
 
diff --git a/Engine/Source/Developer/DesktopPlatform/Public/IDesktopPlatform.h b/Engine/Source/Developer/DesktopPlatform/Public/IDesktopPlatform.h
--- a/Engine/Source/Developer/DesktopPlatform/Public/IDesktopPlatform.h
+++ b/Engine/Source/Developer/DesktopPlatform/Public/IDesktopPlatform.h
@@ -361,6 +361,14 @@ public:
 	*/
 	virtual bool RunUnrealBuildTool(const FText& Description, const FString& RootDir, const FString& Arguments, FFeedbackContext* Warn) = 0;
 
+	/**
+	* Checks if an instance of UnrealBuildTool is running.
+	*
+	* @return true if an instance of UnrealBuildTool is running.
+	*/
+	virtual bool IsUnrealBuildToolRunning() = 0;
+
+
 	/**
 	* Gets the path to the solution for the current project
 	*
diff --git a/Engine/Source/Developer/HotReload/Private/HotReload.cpp b/Engine/Source/Developer/HotReload/Private/HotReload.cpp
--- a/Engine/Source/Developer/HotReload/Private/HotReload.cpp
+++ b/Engine/Source/Developer/HotReload/Private/HotReload.cpp
@@ -838,7 +838,7 @@ bool FHotReloadModule::Tick(float DeltaTime)
 	if (NewModules.Num())
 	{
 		// We have new modules in the queue, but make sure UBT has finished compiling all of them
-		if (!FPlatformProcess::IsApplicationRunning(TEXT("UnrealBuildTool")))
+		if (!FDesktopPlatformModule::Get()->IsUnrealBuildToolRunning())
 		{
 			DoHotReloadFromIDE();
 			NewModules.Empty();
