CIS Fix (missing file from previous check-in).

diff --git a/Engine/Source/Developer/DesktopPlatform/Private/DesktopPlatformBase.h b/Engine/Source/Developer/DesktopPlatform/Private/DesktopPlatformBase.h
--- a/Engine/Source/Developer/DesktopPlatform/Private/DesktopPlatformBase.h
+++ b/Engine/Source/Developer/DesktopPlatform/Private/DesktopPlatformBase.h
@@ -64,7 +64,10 @@ private:
 
 	void GetProjectBuildProducts(const FString& ProjectFileName, TArray<FString> &OutFileNames, TArray<FString> &OutDirectoryNames);
 
-	FString GetUnrealBuildToolExecutableFilename() const;
 	FString GetUnrealBuildToolSourceCodePath() const;
 	bool BuildUnrealBuildTool(FOutputDevice &Ar);
+
+protected:
+
+	FString GetUnrealBuildToolExecutableFilename() const;	
 };
