Support for Linux targets in binary Rocket.

diff --git a/Engine/Source/Developer/DesktopPlatform/Private/PlatformInfo.cpp b/Engine/Source/Developer/DesktopPlatform/Private/PlatformInfo.cpp
--- a/Engine/Source/Developer/DesktopPlatform/Private/PlatformInfo.cpp
+++ b/Engine/Source/Developer/DesktopPlatform/Private/PlatformInfo.cpp
@@ -87,10 +87,10 @@ static const FPlatformInfo PlatformInfoArray[] = {
 	BuildPlatformInfo(TEXT("MacClient"),				TEXT("MacClient"),			LOCTEXT("MacClient", "Mac (Client-only)"),					EPlatformType::Client,		EPlatformFlags::None,			FPlatformIconPaths(TEXT("Launcher/MacTarget/Platform_Mac_24x"), TEXT("Launcher/MacTarget/Platform_Mac_128x")),										TEXT(""),												TEXT(""),			EPlatformSDKStatus::Unknown,	TEXT("/Engine/Tutorial/Installation/InstallingXCodeTutorial.InstallingXCodeTutorial"),						IsAvailableOnMac,							true),
 	BuildPlatformInfo(TEXT("MacServer"),				TEXT("MacServer"),			LOCTEXT("MacServer", "Mac (Dedicated Server)"),				EPlatformType::Server,		EPlatformFlags::None,			FPlatformIconPaths(TEXT("Launcher/MacTarget/Platform_Mac_24x"), TEXT("Launcher/MacTarget/Platform_Mac_128x")),										TEXT(""),												TEXT(""),			EPlatformSDKStatus::Unknown,	TEXT("/Engine/Tutorial/Installation/InstallingXCodeTutorial.InstallingXCodeTutorial"),						IsAvailableOnMac,							true),
 
-	BuildPlatformInfo(TEXT("LinuxNoEditor"),			TEXT("LinuxNoEditor"),		LOCTEXT("LinuxNoEditor", "Linux"),							EPlatformType::Game,		EPlatformFlags::None,			FPlatformIconPaths(TEXT("Launcher/LinuxTarget/Platform_Linux_24x"), TEXT("Launcher/LinuxTarget/Platform_Linux_128x")),								TEXT(""),												TEXT("Linux_x64"),	EPlatformSDKStatus::Unknown,	TEXT(""),																									IsAvailableOnLinux || IsAvailableOnWindows,	false),	
-	BuildPlatformInfo(TEXT("Linux"),					TEXT("Linux"),				LOCTEXT("LinuxEditor", "Linux (Editor)"),					EPlatformType::Editor,		EPlatformFlags::None,			FPlatformIconPaths(TEXT("Launcher/LinuxTarget/Platform_Linux_24x"), TEXT("Launcher/LinuxTarget/Platform_Linux_128x")),								TEXT(""),												TEXT("Linux_x64"),	EPlatformSDKStatus::Unknown,	TEXT(""),																									IsAvailableOnLinux,							false),
-	BuildPlatformInfo(TEXT("LinuxClient"),				TEXT("LinuxClient"),		LOCTEXT("LinuxClient", "Linux (Client-only)"),				EPlatformType::Client,		EPlatformFlags::None,			FPlatformIconPaths(TEXT("Launcher/LinuxTarget/Platform_Linux_24x"), TEXT("Launcher/LinuxTarget/Platform_Linux_128x")),								TEXT(""),												TEXT("Linux_x64"),	EPlatformSDKStatus::Unknown,	TEXT(""),																									IsAvailableOnLinux || IsAvailableOnWindows,	false),
-	BuildPlatformInfo(TEXT("LinuxServer"),				TEXT("LinuxServer"),		LOCTEXT("LinuxServer", "Linux (Dedicated Server)"),			EPlatformType::Server,		EPlatformFlags::None,			FPlatformIconPaths(TEXT("Launcher/LinuxTarget/Platform_Linux_24x"), TEXT("Launcher/LinuxTarget/Platform_Linux_128x")),								TEXT(""),												TEXT("Linux_x64"),	EPlatformSDKStatus::Unknown,	TEXT(""),																									IsAvailableOnLinux || IsAvailableOnWindows,	false),
+	BuildPlatformInfo(TEXT("LinuxNoEditor"),			TEXT("LinuxNoEditor"),		LOCTEXT("LinuxNoEditor", "Linux"),							EPlatformType::Game,		EPlatformFlags::None,			FPlatformIconPaths(TEXT("Launcher/LinuxTarget/Platform_Linux_24x"), TEXT("Launcher/LinuxTarget/Platform_Linux_128x")),								TEXT(""),												TEXT("Linux_x64"),	EPlatformSDKStatus::Unknown,	TEXT(""),																									IsAvailableOnLinux || IsAvailableOnWindows,	 true),	
+	BuildPlatformInfo(TEXT("Linux"),					TEXT("Linux"),				LOCTEXT("LinuxEditor", "Linux (Editor)"),					EPlatformType::Editor,		EPlatformFlags::None,			FPlatformIconPaths(TEXT("Launcher/LinuxTarget/Platform_Linux_24x"), TEXT("Launcher/LinuxTarget/Platform_Linux_128x")),								TEXT(""),												TEXT("Linux_x64"),	EPlatformSDKStatus::Unknown,	TEXT(""),																									IsAvailableOnLinux,							 true),
+	BuildPlatformInfo(TEXT("LinuxClient"),				TEXT("LinuxClient"),		LOCTEXT("LinuxClient", "Linux (Client-only)"),				EPlatformType::Client,		EPlatformFlags::None,			FPlatformIconPaths(TEXT("Launcher/LinuxTarget/Platform_Linux_24x"), TEXT("Launcher/LinuxTarget/Platform_Linux_128x")),								TEXT(""),												TEXT("Linux_x64"),	EPlatformSDKStatus::Unknown,	TEXT(""),																									IsAvailableOnLinux || IsAvailableOnWindows,	 true),
+	BuildPlatformInfo(TEXT("LinuxServer"),				TEXT("LinuxServer"),		LOCTEXT("LinuxServer", "Linux (Dedicated Server)"),			EPlatformType::Server,		EPlatformFlags::None,			FPlatformIconPaths(TEXT("Launcher/LinuxTarget/Platform_Linux_24x"), TEXT("Launcher/LinuxTarget/Platform_Linux_128x")),								TEXT(""),												TEXT("Linux_x64"),	EPlatformSDKStatus::Unknown,	TEXT(""),																									IsAvailableOnLinux || IsAvailableOnWindows,	 true),
 
 	BuildPlatformInfo(TEXT("IOS"),						TEXT("IOS"),				LOCTEXT("IOS", "iOS"),										EPlatformType::Game,		EPlatformFlags::None,			FPlatformIconPaths(TEXT("Launcher/iOSTarget/Platform_iOS_24x"), TEXT("Launcher/iOSTarget/Platform_iOS_128x")),										TEXT(""),												TEXT(""),			EPlatformSDKStatus::Unknown,	IOSTutorial,																								IsAvailableOnWindows || IsAvailableOnMac,	true),
 
diff --git a/Engine/Source/Programs/AutomationTool/Scripts/CopyBuildToStagingDirectory.Automation.cs b/Engine/Source/Programs/AutomationTool/Scripts/CopyBuildToStagingDirectory.Automation.cs
--- a/Engine/Source/Programs/AutomationTool/Scripts/CopyBuildToStagingDirectory.Automation.cs
+++ b/Engine/Source/Programs/AutomationTool/Scripts/CopyBuildToStagingDirectory.Automation.cs
@@ -943,7 +943,7 @@ public partial class Project : CommandUtils
 				ExecutablesToStage,
 				InDedicatedServer,
 				Params.Cook || Params.CookOnTheFly,
-				Params.CrashReporter,
+				Params.CrashReporter && (StagePlatform != UnrealTargetPlatform.Linux || !Params.Rocket),
 				Params.Stage,
 				Params.CookOnTheFly,
 				Params.Archive,
diff --git a/Engine/Source/Programs/AutomationTool/Scripts/GUBP.Automation.cs b/Engine/Source/Programs/AutomationTool/Scripts/GUBP.Automation.cs
--- a/Engine/Source/Programs/AutomationTool/Scripts/GUBP.Automation.cs
+++ b/Engine/Source/Programs/AutomationTool/Scripts/GUBP.Automation.cs
@@ -1330,6 +1330,10 @@ public class GUBP : BuildCommand
                         {
                             continue;
                         }
+                        if (TargetPlatform == UnrealTargetPlatform.Linux && Config != UnrealTargetConfiguration.Shipping && Config != UnrealTargetConfiguration.Development)
+                        {
+                            continue;
+                        }
                     }
                     else if (Config != UnrealTargetConfiguration.Shipping && Config != UnrealTargetConfiguration.Development)
                     {
@@ -1415,6 +1419,10 @@ public class GUBP : BuildCommand
                                     {
                                         continue;
                                     }
+                                    if (TargetPlatform == UnrealTargetPlatform.Linux && Config != UnrealTargetConfiguration.Shipping && Config != UnrealTargetConfiguration.Development)
+                                    {
+                                        continue;
+                                    }
                                 }
                                 else if (Config != UnrealTargetConfiguration.Shipping && Config != UnrealTargetConfiguration.Development)
                                 {
@@ -1431,7 +1439,6 @@ public class GUBP : BuildCommand
                                 Agenda.AddTargets(new string[] { Target.TargetName }, TargetPlatform, Config, GameProj.FilePath, InAddArgs: Args);
                             }
                         }
-
                     }
                 }
             }
@@ -4872,7 +4879,6 @@ public class GUBP : BuildCommand
             foreach (var Plat in ActivePlatforms)
             {
                 if (
-                    Plat != UnrealTargetPlatform.Linux &&
                     (Plat != UnrealTargetPlatform.PS4 || ParseParam("WithPS4")) &&
                     Plat != UnrealTargetPlatform.WinRT &&
                     Plat != UnrealTargetPlatform.WinRT_ARM &&
diff --git a/Engine/Source/Programs/UnrealBuildTool/Linux/UEBuildLinux.cs b/Engine/Source/Programs/UnrealBuildTool/Linux/UEBuildLinux.cs
--- a/Engine/Source/Programs/UnrealBuildTool/Linux/UEBuildLinux.cs
+++ b/Engine/Source/Programs/UnrealBuildTool/Linux/UEBuildLinux.cs
@@ -106,12 +106,6 @@ namespace UnrealBuildTool
          */
         protected override void RegisterBuildPlatformInternal()
         {
-			[REDACTED_EMAIL]: Add platform support
-			if (UnrealBuildTool.RunningRocket() || UnrealBuildTool.BuildingRocket())
-			{
-				return;
-			}
-
 			if ((ProjectFileGenerator.bGenerateProjectFiles == true) || (HasRequiredSDKsInstalled() == SDKStatus.Valid))
             {
                 bool bRegisterBuildPlatform = true;
