Disable HTML5 in Rocket.

diff --git a/Engine/Source/Developer/DesktopPlatform/Private/PlatformInfo.cpp b/Engine/Source/Developer/DesktopPlatform/Private/PlatformInfo.cpp
--- a/Engine/Source/Developer/DesktopPlatform/Private/PlatformInfo.cpp
+++ b/Engine/Source/Developer/DesktopPlatform/Private/PlatformInfo.cpp
@@ -104,7 +104,7 @@ static const FPlatformInfo PlatformInfoArray[] = {
 	BuildPlatformInfo(TEXT("Android_ES31"),				TEXT("Android_ES31"),		LOCTEXT("Android_ES31", "Android (ES31)"),					EPlatformType::Game,		EPlatformFlags::CookFlavor,		FPlatformIconPaths(TEXT("Launcher/AndroidTarget/Platform_Android_ES31_24x"), TEXT("Launcher/AndroidTarget/Platform_Android_128x")),					TEXT("-targetplatform=Android -cookflavor=ES31"),		TEXT("Android"),	EPlatformSDKStatus::Unknown,	TEXT("Shared/Tutorials/SettingUpAndroidTutorial"),			IsAvailableOnWindows,						true),
 //	BuildPlatformInfo(TEXT("Android_GL4"),				TEXT("Android_GL4"),		LOCTEXT("Android_GL4", "Android (GL4)"),					EPlatformType::Game,		EPlatformFlags::CookFlavor,		FPlatformIconPaths(TEXT("Launcher/AndroidTarget/Platform_Android_GL4_24x"), TEXT("Launcher/AndroidTarget/Platform_Android_128x")),					TEXT("-targetplatform=Android -cookflavor=GL4"),		TEXT("Android"),	EPlatformSDKStatus::Unknown,	TEXT("Shared/Tutorials/SettingUpAndroidTutorial"),			IsAvailableOnWindows,						true),
 
-	BuildPlatformInfo(TEXT("HTML5"),					TEXT("HTML5"),				LOCTEXT("HTML5", "HTML5"),									EPlatformType::Game,		EPlatformFlags::None,			FPlatformIconPaths(TEXT("Launcher/HTML5Target/Platform_HTML5_24x"), TEXT("Launcher/HTML5Target/Platform_HTML5_128x")),								TEXT(""),												TEXT("HTML5"),		EPlatformSDKStatus::Unknown,	TEXT(""),																									IsAvailableOnWindows || IsAvailableOnMac,   IsAvailableOnWindows),
+	BuildPlatformInfo(TEXT("HTML5"),					TEXT("HTML5"),				LOCTEXT("HTML5", "HTML5"),									EPlatformType::Game,		EPlatformFlags::None,			FPlatformIconPaths(TEXT("Launcher/HTML5Target/Platform_HTML5_24x"), TEXT("Launcher/HTML5Target/Platform_HTML5_128x")),								TEXT(""),												TEXT("HTML5"),		EPlatformSDKStatus::Unknown,	TEXT(""),																									IsAvailableOnWindows || IsAvailableOnMac,   false),
 
 	BuildPlatformInfo(TEXT("PS4"),						TEXT("PS4"),				LOCTEXT("PS4", "PlayStation 4"),							EPlatformType::Game,		EPlatformFlags::None,			FPlatformIconPaths(TEXT("Launcher/PS4Target/Platform_PS4_24x"), TEXT("Launcher/PS4Target/Platform_PS4_128x")),										TEXT(""),												TEXT("PS4"),		EPlatformSDKStatus::Unknown,	TEXT(""),																									IsAvailableOnWindows,						false),
 
diff --git a/Engine/Source/Developer/TargetPlatform/TargetPlatform.Build.cs b/Engine/Source/Developer/TargetPlatform/TargetPlatform.Build.cs
--- a/Engine/Source/Developer/TargetPlatform/TargetPlatform.Build.cs
+++ b/Engine/Source/Developer/TargetPlatform/TargetPlatform.Build.cs
@@ -52,7 +52,11 @@ public class TargetPlatform : ModuleRules
 					DynamicallyLoadedModuleNames.Add("Android_ETC1TargetPlatform");
 					DynamicallyLoadedModuleNames.Add("Android_ETC2TargetPlatform");
 					DynamicallyLoadedModuleNames.Add("IOSTargetPlatform");
-                    DynamicallyLoadedModuleNames.Add("HTML5TargetPlatform");
+
+					if (!UnrealBuildTool.UnrealBuildTool.RunningRocket() && !UnrealBuildTool.UnrealBuildTool.BuildingRocket())
+					{
+						DynamicallyLoadedModuleNames.Add("HTML5TargetPlatform");
+					}
 				}
 			}
 			else if (Target.Platform == UnrealTargetPlatform.Mac)
diff --git a/Engine/Source/Programs/AutomationTool/Scripts/GUBP.Automation.cs b/Engine/Source/Programs/AutomationTool/Scripts/GUBP.Automation.cs
--- a/Engine/Source/Programs/AutomationTool/Scripts/GUBP.Automation.cs
+++ b/Engine/Source/Programs/AutomationTool/Scripts/GUBP.Automation.cs
@@ -4878,7 +4878,7 @@ public class GUBP : BuildCommand
             var FilteredActivePlatforms = new List<UnrealTargetPlatform>();
             foreach (var Plat in ActivePlatforms)
             {
-                if (
+                if (Plat != UnrealTargetPlatform.HTML5 && 
                     (Plat != UnrealTargetPlatform.PS4 || ParseParam("WithPS4")) &&
                     Plat != UnrealTargetPlatform.WinRT &&
                     Plat != UnrealTargetPlatform.WinRT_ARM &&
diff --git a/Engine/Source/Programs/UnrealBuildTool/HTML5/UEBuildHTML5.cs b/Engine/Source/Programs/UnrealBuildTool/HTML5/UEBuildHTML5.cs
--- a/Engine/Source/Programs/UnrealBuildTool/HTML5/UEBuildHTML5.cs
+++ b/Engine/Source/Programs/UnrealBuildTool/HTML5/UEBuildHTML5.cs
@@ -118,6 +118,12 @@ namespace UnrealBuildTool
          */
         protected override void RegisterBuildPlatformInternal()
         {
+            [REDACTED_EMAIL]: Add platform support
+            if (UnrealBuildTool.RunningRocket() || UnrealBuildTool.BuildingRocket())
+            {
+                return;
+            }
+
             // Make sure the SDK is installed
             if ((ProjectFileGenerator.bGenerateProjectFiles == true) || (HasRequiredSDKsInstalled() == SDKStatus.Valid))
             {
