Re-enable documentation link for HTML5, PS4 & XboxOne 'missing SDK' when packaging

Added a fallback to treat the path passed-in as a documentation path if the asset path doesnt correspond to a BP.

diff --git a/Engine/Source/Developer/DesktopPlatform/Private/PlatformInfo.cpp b/Engine/Source/Developer/DesktopPlatform/Private/PlatformInfo.cpp
--- a/Engine/Source/Developer/DesktopPlatform/Private/PlatformInfo.cpp
+++ b/Engine/Source/Developer/DesktopPlatform/Private/PlatformInfo.cpp
@@ -75,40 +75,40 @@ static const FString IOSTutorial = TEXT("/Engine/Tutorial/NotYetImplemented");
 
 static const FPlatformInfo PlatformInfoArray[] = {
 	//				  PlatformInfoName					TargetPlatformName			DisplayName													PlatformType				PlatformFlags					IconPaths																																			UATCommandLine											AutoSDKPath			SDKStatus						SDKTutorial																								EnabledForUse									EnabledForBinary
-	BuildPlatformInfo(TEXT("WindowsNoEditor"),			TEXT("WindowsNoEditor"),	LOCTEXT("WindowsNoEditor", "Windows"),						EPlatformType::Game,		EPlatformFlags::None,			FPlatformIconPaths(TEXT("Launcher/WindowsTarget/Platform_WindowsNoEditor_24x"), TEXT("Launcher/WindowsTarget/Platform_WindowsNoEditor_128x")),		TEXT(""),												TEXT(""),			EPlatformSDKStatus::Unknown,	TEXT("/Engine/Tutorial/Installation/InstallingVisualStudioTutorial.InstallingVisualStudioTutorial"),		IsAvailableOnWindows,						true),
-	BuildPlatformInfo(TEXT("WindowsNoEditor_Win32"),	TEXT("WindowsNoEditor"),	LOCTEXT("WindowsNoEditor_Win32", "Windows (32-bit)"),		EPlatformType::Game,		EPlatformFlags::BuildFlavor,	FPlatformIconPaths(TEXT("Launcher/WindowsTarget/Platform_WindowsNoEditor_24x"), TEXT("Launcher/WindowsTarget/Platform_WindowsNoEditor_128x")),		TEXT("-targetplatform=Win32"),							TEXT(""),			EPlatformSDKStatus::Unknown,	TEXT("/Engine/Tutorial/Installation/InstallingVisualStudioTutorial.InstallingVisualStudioTutorial"),		IsAvailableOnWindows,						true),
-	BuildPlatformInfo(TEXT("WindowsNoEditor_Win64"),	TEXT("WindowsNoEditor"),	LOCTEXT("WindowsNoEditor_Win64", "Windows (64-bit)"),		EPlatformType::Game,		EPlatformFlags::BuildFlavor,	FPlatformIconPaths(TEXT("Launcher/WindowsTarget/Platform_WindowsNoEditor_24x"), TEXT("Launcher/WindowsTarget/Platform_WindowsNoEditor_128x")),		TEXT("-targetplatform=Win64"),							TEXT(""),			EPlatformSDKStatus::Unknown,	TEXT("/Engine/Tutorial/Installation/InstallingVisualStudioTutorial.InstallingVisualStudioTutorial"),		IsAvailableOnWindows,						true),
-	BuildPlatformInfo(TEXT("Windows"),					TEXT("Windows"),			LOCTEXT("WindowsEditor", "Windows (Editor)"),				EPlatformType::Editor,		EPlatformFlags::None,			FPlatformIconPaths(TEXT("Launcher/WindowsTarget/Platform_Windows_24x"), TEXT("Launcher/WindowsTarget/Platform_Windows_128x")),						TEXT(""),												TEXT(""),			EPlatformSDKStatus::Unknown,	TEXT("/Engine/Tutorial/Installation/InstallingVisualStudioTutorial.InstallingVisualStudioTutorial"),		IsAvailableOnWindows,						true),
-	BuildPlatformInfo(TEXT("WindowsClient"),			TEXT("WindowsClient"),		LOCTEXT("WindowsClient", "Windows (Client-only)"),			EPlatformType::Client,		EPlatformFlags::None,			FPlatformIconPaths(TEXT("Launcher/WindowsTarget/Platform_Windows_24x"), TEXT("Launcher/WindowsTarget/Platform_Windows_128x")),						TEXT(""),												TEXT(""),			EPlatformSDKStatus::Unknown,	TEXT("/Engine/Tutorial/Installation/InstallingVisualStudioTutorial.InstallingVisualStudioTutorial"),		IsAvailableOnWindows,						true),
-	BuildPlatformInfo(TEXT("WindowsServer"),			TEXT("WindowsServer"),		LOCTEXT("WindowsServer", "Windows (Dedicated Server)"),		EPlatformType::Server,		EPlatformFlags::None,			FPlatformIconPaths(TEXT("Launcher/WindowsTarget/Platform_WindowsServer_24x"), TEXT("Launcher/WindowsTarget/Platform_WindowsServer_128x")),			TEXT(""),												TEXT(""),			EPlatformSDKStatus::Unknown,	TEXT("/Engine/Tutorial/Installation/InstallingVisualStudioTutorial.InstallingVisualStudioTutorial"),		IsAvailableOnWindows,						true),
-
-	BuildPlatformInfo(TEXT("MacNoEditor"),				TEXT("MacNoEditor"),		LOCTEXT("MacNoEditor", "Mac"),								EPlatformType::Game,		EPlatformFlags::None,			FPlatformIconPaths(TEXT("Launcher/MacTarget/Platform_Mac_24x"), TEXT("Launcher/MacTarget/Platform_Mac_128x")),										TEXT(""),												TEXT(""),			EPlatformSDKStatus::Unknown,	TEXT("/Engine/Tutorial/Installation/InstallingXCodeTutorial.InstallingXCodeTutorial"),						IsAvailableOnMac,							true),
-	BuildPlatformInfo(TEXT("Mac"),						TEXT("Mac"),				LOCTEXT("MacEditor", "Mac (Editor)"),						EPlatformType::Editor,		EPlatformFlags::None,			FPlatformIconPaths(TEXT("Launcher/MacTarget/Platform_Mac_24x"), TEXT("Launcher/MacTarget/Platform_Mac_128x")),										TEXT(""),												TEXT(""),			EPlatformSDKStatus::Unknown,	TEXT("/Engine/Tutorial/Installation/InstallingXCodeTutorial.InstallingXCodeTutorial"),						IsAvailableOnMac,							true),
-	BuildPlatformInfo(TEXT("MacClient"),				TEXT("MacClient"),			LOCTEXT("MacClient", "Mac (Client-only)"),					EPlatformType::Client,		EPlatformFlags::None,			FPlatformIconPaths(TEXT("Launcher/MacTarget/Platform_Mac_24x"), TEXT("Launcher/MacTarget/Platform_Mac_128x")),										TEXT(""),												TEXT(""),			EPlatformSDKStatus::Unknown,	TEXT("/Engine/Tutorial/Installation/InstallingXCodeTutorial.InstallingXCodeTutorial"),						IsAvailableOnMac,							true),
-	BuildPlatformInfo(TEXT("MacServer"),				TEXT("MacServer"),			LOCTEXT("MacServer", "Mac (Dedicated Server)"),				EPlatformType::Server,		EPlatformFlags::None,			FPlatformIconPaths(TEXT("Launcher/MacTarget/Platform_Mac_24x"), TEXT("Launcher/MacTarget/Platform_Mac_128x")),										TEXT(""),												TEXT(""),			EPlatformSDKStatus::Unknown,	TEXT("/Engine/Tutorial/Installation/InstallingXCodeTutorial.InstallingXCodeTutorial"),						IsAvailableOnMac,							true),
-
-	BuildPlatformInfo(TEXT("LinuxNoEditor"),			TEXT("LinuxNoEditor"),		LOCTEXT("LinuxNoEditor", "Linux"),							EPlatformType::Game,		EPlatformFlags::None,			FPlatformIconPaths(TEXT("Launcher/LinuxTarget/Platform_Linux_24x"), TEXT("Launcher/LinuxTarget/Platform_Linux_128x")),								TEXT(""),												TEXT("Linux_x64"),	EPlatformSDKStatus::Unknown,	TEXT(""),																									IsAvailableOnLinux || IsAvailableOnWindows,	 true),	
-	BuildPlatformInfo(TEXT("Linux"),					TEXT("Linux"),				LOCTEXT("LinuxEditor", "Linux (Editor)"),					EPlatformType::Editor,		EPlatformFlags::None,			FPlatformIconPaths(TEXT("Launcher/LinuxTarget/Platform_Linux_24x"), TEXT("Launcher/LinuxTarget/Platform_Linux_128x")),								TEXT(""),												TEXT("Linux_x64"),	EPlatformSDKStatus::Unknown,	TEXT(""),																									IsAvailableOnLinux,							 true),
-	BuildPlatformInfo(TEXT("LinuxClient"),				TEXT("LinuxClient"),		LOCTEXT("LinuxClient", "Linux (Client-only)"),				EPlatformType::Client,		EPlatformFlags::None,			FPlatformIconPaths(TEXT("Launcher/LinuxTarget/Platform_Linux_24x"), TEXT("Launcher/LinuxTarget/Platform_Linux_128x")),								TEXT(""),												TEXT("Linux_x64"),	EPlatformSDKStatus::Unknown,	TEXT(""),																									IsAvailableOnLinux || IsAvailableOnWindows,	 true),
-	BuildPlatformInfo(TEXT("LinuxServer"),				TEXT("LinuxServer"),		LOCTEXT("LinuxServer", "Linux (Dedicated Server)"),			EPlatformType::Server,		EPlatformFlags::None,			FPlatformIconPaths(TEXT("Launcher/LinuxTarget/Platform_Linux_24x"), TEXT("Launcher/LinuxTarget/Platform_Linux_128x")),								TEXT(""),												TEXT("Linux_x64"),	EPlatformSDKStatus::Unknown,	TEXT(""),																									IsAvailableOnLinux || IsAvailableOnWindows,	 true),
-
-	BuildPlatformInfo(TEXT("IOS"),						TEXT("IOS"),				LOCTEXT("IOS", "iOS"),										EPlatformType::Game,		EPlatformFlags::None,			FPlatformIconPaths(TEXT("Launcher/iOSTarget/Platform_iOS_24x"), TEXT("Launcher/iOSTarget/Platform_iOS_128x")),										TEXT(""),												TEXT(""),			EPlatformSDKStatus::Unknown,	IOSTutorial,																								IsAvailableOnWindows || IsAvailableOnMac,	true),
-
-	BuildPlatformInfo(TEXT("Android"),					TEXT("Android"),			LOCTEXT("Android", "Android"),								EPlatformType::Game,		EPlatformFlags::None,			FPlatformIconPaths(TEXT("Launcher/AndroidTarget/Platform_Android_24x"), TEXT("Launcher/AndroidTarget/Platform_Android_128x")),						TEXT(""),												TEXT("Android"),	EPlatformSDKStatus::Unknown,	TEXT("/Engine/Tutorial/Mobile/SettingUpAndroidTutorial.SettingUpAndroidTutorial"),							IsAvailableOnWindows,						true),
-	BuildPlatformInfo(TEXT("Android_All"),				TEXT("Android"),			LOCTEXT("Android_All", "Android (All)"),					EPlatformType::Game,		EPlatformFlags::CookFlavor,		FPlatformIconPaths(TEXT("Launcher/AndroidTarget/Platform_Android_24x"), TEXT("Launcher/AndroidTarget/Platform_Android_128x")),						TEXT(""),												TEXT("Android"),	EPlatformSDKStatus::Unknown,	TEXT("/Engine/Tutorial/Mobile/SettingUpAndroidTutorial.SettingUpAndroidTutorial"),							IsAvailableOnWindows,						true),
-	BuildPlatformInfo(TEXT("Android_ATC"),				TEXT("Android_ATC"),		LOCTEXT("Android_ATC", "Android (ATC)"),					EPlatformType::Game,		EPlatformFlags::CookFlavor,		FPlatformIconPaths(TEXT("Launcher/AndroidTarget/Platform_Android_ATC_24x"), TEXT("Launcher/AndroidTarget/Platform_Android_128x")),					TEXT("-targetplatform=Android -cookflavor=ATC"),		TEXT("Android"),	EPlatformSDKStatus::Unknown,	TEXT("/Engine/Tutorial/Mobile/SettingUpAndroidTutorial.SettingUpAndroidTutorial"),							IsAvailableOnWindows,						true),
-	BuildPlatformInfo(TEXT("Android_DXT"),				TEXT("Android_DXT"),		LOCTEXT("Android_DXT", "Android (DXT)"),					EPlatformType::Game,		EPlatformFlags::CookFlavor,		FPlatformIconPaths(TEXT("Launcher/AndroidTarget/Platform_Android_DXT_24x"), TEXT("Launcher/AndroidTarget/Platform_Android_128x")),					TEXT("-targetplatform=Android -cookflavor=DXT"),		TEXT("Android"),	EPlatformSDKStatus::Unknown,	TEXT("/Engine/Tutorial/Mobile/SettingUpAndroidTutorial.SettingUpAndroidTutorial"),							IsAvailableOnWindows,						true),
-	BuildPlatformInfo(TEXT("Android_ETC1"),				TEXT("Android_ETC1"),		LOCTEXT("Android_ETC1", "Android (ETC1)"),					EPlatformType::Game,		EPlatformFlags::CookFlavor,		FPlatformIconPaths(TEXT("Launcher/AndroidTarget/Platform_Android_ETC1_24x"), TEXT("Launcher/AndroidTarget/Platform_Android_128x")),					TEXT("-targetplatform=Android -cookflavor=ETC1"),		TEXT("Android"),	EPlatformSDKStatus::Unknown,	TEXT("/Engine/Tutorial/Mobile/SettingUpAndroidTutorial.SettingUpAndroidTutorial"),							IsAvailableOnWindows,						true),
-	BuildPlatformInfo(TEXT("Android_ETC2"),				TEXT("Android_ETC2"),		LOCTEXT("Android_ETC2", "Android (ETC2)"),					EPlatformType::Game,		EPlatformFlags::CookFlavor,		FPlatformIconPaths(TEXT("Launcher/AndroidTarget/Platform_Android_ETC2_24x"), TEXT("Launcher/AndroidTarget/Platform_Android_128x")),					TEXT("-targetplatform=Android -cookflavor=ETC2"),		TEXT("Android"),	EPlatformSDKStatus::Unknown,	TEXT("/Engine/Tutorial/Mobile/SettingUpAndroidTutorial.SettingUpAndroidTutorial"),							IsAvailableOnWindows,						true),
-	BuildPlatformInfo(TEXT("Android_PVRTC"),			TEXT("Android_PVRTC"),		LOCTEXT("Android_PVRTC", "Android (PVRTC)"),				EPlatformType::Game,		EPlatformFlags::CookFlavor,		FPlatformIconPaths(TEXT("Launcher/AndroidTarget/Platform_Android_PVRTC_24x"), TEXT("Launcher/AndroidTarget/Platform_Android_128x")),				TEXT("-targetplatform=Android -cookflavor=PVRTC"),		TEXT("Android"),	EPlatformSDKStatus::Unknown,	TEXT("/Engine/Tutorial/Mobile/SettingUpAndroidTutorial.SettingUpAndroidTutorial"),							IsAvailableOnWindows,						true),
-	BuildPlatformInfo(TEXT("Android_ES31"),				TEXT("Android_ES31"),		LOCTEXT("Android_ES31", "Android (ES31)"),					EPlatformType::Game,		EPlatformFlags::CookFlavor,		FPlatformIconPaths(TEXT("Launcher/AndroidTarget/Platform_Android_ES31_24x"), TEXT("Launcher/AndroidTarget/Platform_Android_128x")),					TEXT("-targetplatform=Android -cookflavor=ES31"),		TEXT("Android"),	EPlatformSDKStatus::Unknown,	TEXT("Shared/Tutorials/SettingUpAndroidTutorial"),			IsAvailableOnWindows,						true),
-//	BuildPlatformInfo(TEXT("Android_GL4"),				TEXT("Android_GL4"),		LOCTEXT("Android_GL4", "Android (GL4)"),					EPlatformType::Game,		EPlatformFlags::CookFlavor,		FPlatformIconPaths(TEXT("Launcher/AndroidTarget/Platform_Android_GL4_24x"), TEXT("Launcher/AndroidTarget/Platform_Android_128x")),					TEXT("-targetplatform=Android -cookflavor=GL4"),		TEXT("Android"),	EPlatformSDKStatus::Unknown,	TEXT("Shared/Tutorials/SettingUpAndroidTutorial"),			IsAvailableOnWindows,						true),
-
-	BuildPlatformInfo(TEXT("HTML5"),					TEXT("HTML5"),				LOCTEXT("HTML5", "HTML5"),									EPlatformType::Game,		EPlatformFlags::None,			FPlatformIconPaths(TEXT("Launcher/HTML5Target/Platform_HTML5_24x"), TEXT("Launcher/HTML5Target/Platform_HTML5_128x")),								TEXT(""),												TEXT("HTML5"),		EPlatformSDKStatus::Unknown,	TEXT(""),																									IsAvailableOnWindows || IsAvailableOnMac,   false),
-
-	BuildPlatformInfo(TEXT("PS4"),						TEXT("PS4"),				LOCTEXT("PS4", "PlayStation 4"),							EPlatformType::Game,		EPlatformFlags::None,			FPlatformIconPaths(TEXT("Launcher/PS4Target/Platform_PS4_24x"), TEXT("Launcher/PS4Target/Platform_PS4_128x")),										TEXT(""),												TEXT("PS4"),		EPlatformSDKStatus::Unknown,	TEXT(""),																									IsAvailableOnWindows,						false),
-
-	BuildPlatformInfo(TEXT("XboxOne"),					TEXT("XboxOne"),			LOCTEXT("XboxOne", "Xbox One"),								EPlatformType::Game,		EPlatformFlags::None,			FPlatformIconPaths(TEXT("Launcher/XboxOneTarget/Platform_XboxOne_24x"), TEXT("Launcher/XboxOneTarget/Platform_XboxOne_128x")),						TEXT(""),												TEXT("XboxOne"),	EPlatformSDKStatus::Unknown,	TEXT(""),																									IsAvailableOnWindows,						false),
+	BuildPlatformInfo(TEXT("WindowsNoEditor"),			TEXT("WindowsNoEditor"),	LOCTEXT("WindowsNoEditor", "Windows"),						EPlatformType::Game,		EPlatformFlags::None,			FPlatformIconPaths(TEXT("Launcher/WindowsTarget/Platform_WindowsNoEditor_24x"), TEXT("Launcher/WindowsTarget/Platform_WindowsNoEditor_128x")),		TEXT(""),												TEXT(""),			EPlatformSDKStatus::Unknown,	TEXT("/Engine/Tutorial/Installation/InstallingVisualStudioTutorial.InstallingVisualStudioTutorial"),	IsAvailableOnWindows,						true),
+	BuildPlatformInfo(TEXT("WindowsNoEditor_Win32"),	TEXT("WindowsNoEditor"),	LOCTEXT("WindowsNoEditor_Win32", "Windows (32-bit)"),		EPlatformType::Game,		EPlatformFlags::BuildFlavor,	FPlatformIconPaths(TEXT("Launcher/WindowsTarget/Platform_WindowsNoEditor_24x"), TEXT("Launcher/WindowsTarget/Platform_WindowsNoEditor_128x")),		TEXT("-targetplatform=Win32"),							TEXT(""),			EPlatformSDKStatus::Unknown,	TEXT("/Engine/Tutorial/Installation/InstallingVisualStudioTutorial.InstallingVisualStudioTutorial"),	IsAvailableOnWindows,						true),
+	BuildPlatformInfo(TEXT("WindowsNoEditor_Win64"),	TEXT("WindowsNoEditor"),	LOCTEXT("WindowsNoEditor_Win64", "Windows (64-bit)"),		EPlatformType::Game,		EPlatformFlags::BuildFlavor,	FPlatformIconPaths(TEXT("Launcher/WindowsTarget/Platform_WindowsNoEditor_24x"), TEXT("Launcher/WindowsTarget/Platform_WindowsNoEditor_128x")),		TEXT("-targetplatform=Win64"),							TEXT(""),			EPlatformSDKStatus::Unknown,	TEXT("/Engine/Tutorial/Installation/InstallingVisualStudioTutorial.InstallingVisualStudioTutorial"),	IsAvailableOnWindows,						true),
+	BuildPlatformInfo(TEXT("Windows"),					TEXT("Windows"),			LOCTEXT("WindowsEditor", "Windows (Editor)"),				EPlatformType::Editor,		EPlatformFlags::None,			FPlatformIconPaths(TEXT("Launcher/WindowsTarget/Platform_Windows_24x"), TEXT("Launcher/WindowsTarget/Platform_Windows_128x")),						TEXT(""),												TEXT(""),			EPlatformSDKStatus::Unknown,	TEXT("/Engine/Tutorial/Installation/InstallingVisualStudioTutorial.InstallingVisualStudioTutorial"),	IsAvailableOnWindows,						true),
+	BuildPlatformInfo(TEXT("WindowsClient"),			TEXT("WindowsClient"),		LOCTEXT("WindowsClient", "Windows (Client-only)"),			EPlatformType::Client,		EPlatformFlags::None,			FPlatformIconPaths(TEXT("Launcher/WindowsTarget/Platform_Windows_24x"), TEXT("Launcher/WindowsTarget/Platform_Windows_128x")),						TEXT(""),												TEXT(""),			EPlatformSDKStatus::Unknown,	TEXT("/Engine/Tutorial/Installation/InstallingVisualStudioTutorial.InstallingVisualStudioTutorial"),	IsAvailableOnWindows,						true),
+	BuildPlatformInfo(TEXT("WindowsServer"),			TEXT("WindowsServer"),		LOCTEXT("WindowsServer", "Windows (Dedicated Server)"),		EPlatformType::Server,		EPlatformFlags::None,			FPlatformIconPaths(TEXT("Launcher/WindowsTarget/Platform_WindowsServer_24x"), TEXT("Launcher/WindowsTarget/Platform_WindowsServer_128x")),			TEXT(""),												TEXT(""),			EPlatformSDKStatus::Unknown,	TEXT("/Engine/Tutorial/Installation/InstallingVisualStudioTutorial.InstallingVisualStudioTutorial"),	IsAvailableOnWindows,						true),
+
+	BuildPlatformInfo(TEXT("MacNoEditor"),				TEXT("MacNoEditor"),		LOCTEXT("MacNoEditor", "Mac"),								EPlatformType::Game,		EPlatformFlags::None,			FPlatformIconPaths(TEXT("Launcher/MacTarget/Platform_Mac_24x"), TEXT("Launcher/MacTarget/Platform_Mac_128x")),										TEXT(""),												TEXT(""),			EPlatformSDKStatus::Unknown,	TEXT("/Engine/Tutorial/Installation/InstallingXCodeTutorial.InstallingXCodeTutorial"),					IsAvailableOnMac,							true),
+	BuildPlatformInfo(TEXT("Mac"),						TEXT("Mac"),				LOCTEXT("MacEditor", "Mac (Editor)"),						EPlatformType::Editor,		EPlatformFlags::None,			FPlatformIconPaths(TEXT("Launcher/MacTarget/Platform_Mac_24x"), TEXT("Launcher/MacTarget/Platform_Mac_128x")),										TEXT(""),												TEXT(""),			EPlatformSDKStatus::Unknown,	TEXT("/Engine/Tutorial/Installation/InstallingXCodeTutorial.InstallingXCodeTutorial"),					IsAvailableOnMac,							true),
+	BuildPlatformInfo(TEXT("MacClient"),				TEXT("MacClient"),			LOCTEXT("MacClient", "Mac (Client-only)"),					EPlatformType::Client,		EPlatformFlags::None,			FPlatformIconPaths(TEXT("Launcher/MacTarget/Platform_Mac_24x"), TEXT("Launcher/MacTarget/Platform_Mac_128x")),										TEXT(""),												TEXT(""),			EPlatformSDKStatus::Unknown,	TEXT("/Engine/Tutorial/Installation/InstallingXCodeTutorial.InstallingXCodeTutorial"),					IsAvailableOnMac,							true),
+	BuildPlatformInfo(TEXT("MacServer"),				TEXT("MacServer"),			LOCTEXT("MacServer", "Mac (Dedicated Server)"),				EPlatformType::Server,		EPlatformFlags::None,			FPlatformIconPaths(TEXT("Launcher/MacTarget/Platform_Mac_24x"), TEXT("Launcher/MacTarget/Platform_Mac_128x")),										TEXT(""),												TEXT(""),			EPlatformSDKStatus::Unknown,	TEXT("/Engine/Tutorial/Installation/InstallingXCodeTutorial.InstallingXCodeTutorial"),					IsAvailableOnMac,							true),
+
+	BuildPlatformInfo(TEXT("LinuxNoEditor"),			TEXT("LinuxNoEditor"),		LOCTEXT("LinuxNoEditor", "Linux"),							EPlatformType::Game,		EPlatformFlags::None,			FPlatformIconPaths(TEXT("Launcher/LinuxTarget/Platform_Linux_24x"), TEXT("Launcher/LinuxTarget/Platform_Linux_128x")),								TEXT(""),												TEXT("Linux_x64"),	EPlatformSDKStatus::Unknown,	TEXT(""),																								IsAvailableOnLinux || IsAvailableOnWindows,	 true),	
+	BuildPlatformInfo(TEXT("Linux"),					TEXT("Linux"),				LOCTEXT("LinuxEditor", "Linux (Editor)"),					EPlatformType::Editor,		EPlatformFlags::None,			FPlatformIconPaths(TEXT("Launcher/LinuxTarget/Platform_Linux_24x"), TEXT("Launcher/LinuxTarget/Platform_Linux_128x")),								TEXT(""),												TEXT("Linux_x64"),	EPlatformSDKStatus::Unknown,	TEXT(""),																								IsAvailableOnLinux,							 true),
+	BuildPlatformInfo(TEXT("LinuxClient"),				TEXT("LinuxClient"),		LOCTEXT("LinuxClient", "Linux (Client-only)"),				EPlatformType::Client,		EPlatformFlags::None,			FPlatformIconPaths(TEXT("Launcher/LinuxTarget/Platform_Linux_24x"), TEXT("Launcher/LinuxTarget/Platform_Linux_128x")),								TEXT(""),												TEXT("Linux_x64"),	EPlatformSDKStatus::Unknown,	TEXT(""),																								IsAvailableOnLinux || IsAvailableOnWindows,	 true),
+	BuildPlatformInfo(TEXT("LinuxServer"),				TEXT("LinuxServer"),		LOCTEXT("LinuxServer", "Linux (Dedicated Server)"),			EPlatformType::Server,		EPlatformFlags::None,			FPlatformIconPaths(TEXT("Launcher/LinuxTarget/Platform_Linux_24x"), TEXT("Launcher/LinuxTarget/Platform_Linux_128x")),								TEXT(""),												TEXT("Linux_x64"),	EPlatformSDKStatus::Unknown,	TEXT(""),																								IsAvailableOnLinux || IsAvailableOnWindows,	 true),
+
+	BuildPlatformInfo(TEXT("IOS"),						TEXT("IOS"),				LOCTEXT("IOS", "iOS"),										EPlatformType::Game,		EPlatformFlags::None,			FPlatformIconPaths(TEXT("Launcher/iOSTarget/Platform_iOS_24x"), TEXT("Launcher/iOSTarget/Platform_iOS_128x")),										TEXT(""),												TEXT(""),			EPlatformSDKStatus::Unknown,	IOSTutorial,																							IsAvailableOnWindows || IsAvailableOnMac,	true),
+
+	BuildPlatformInfo(TEXT("Android"),					TEXT("Android"),			LOCTEXT("Android", "Android"),								EPlatformType::Game,		EPlatformFlags::None,			FPlatformIconPaths(TEXT("Launcher/AndroidTarget/Platform_Android_24x"), TEXT("Launcher/AndroidTarget/Platform_Android_128x")),						TEXT(""),												TEXT("Android"),	EPlatformSDKStatus::Unknown,	TEXT("/Engine/Tutorial/Mobile/SettingUpAndroidTutorial.SettingUpAndroidTutorial"),						IsAvailableOnWindows,						true),
+	BuildPlatformInfo(TEXT("Android_All"),				TEXT("Android"),			LOCTEXT("Android_All", "Android (All)"),					EPlatformType::Game,		EPlatformFlags::CookFlavor,		FPlatformIconPaths(TEXT("Launcher/AndroidTarget/Platform_Android_24x"), TEXT("Launcher/AndroidTarget/Platform_Android_128x")),						TEXT(""),												TEXT("Android"),	EPlatformSDKStatus::Unknown,	TEXT("/Engine/Tutorial/Mobile/SettingUpAndroidTutorial.SettingUpAndroidTutorial"),						IsAvailableOnWindows,						true),
+	BuildPlatformInfo(TEXT("Android_ATC"),				TEXT("Android_ATC"),		LOCTEXT("Android_ATC", "Android (ATC)"),					EPlatformType::Game,		EPlatformFlags::CookFlavor,		FPlatformIconPaths(TEXT("Launcher/AndroidTarget/Platform_Android_ATC_24x"), TEXT("Launcher/AndroidTarget/Platform_Android_128x")),					TEXT("-targetplatform=Android -cookflavor=ATC"),		TEXT("Android"),	EPlatformSDKStatus::Unknown,	TEXT("/Engine/Tutorial/Mobile/SettingUpAndroidTutorial.SettingUpAndroidTutorial"),						IsAvailableOnWindows,						true),
+	BuildPlatformInfo(TEXT("Android_DXT"),				TEXT("Android_DXT"),		LOCTEXT("Android_DXT", "Android (DXT)"),					EPlatformType::Game,		EPlatformFlags::CookFlavor,		FPlatformIconPaths(TEXT("Launcher/AndroidTarget/Platform_Android_DXT_24x"), TEXT("Launcher/AndroidTarget/Platform_Android_128x")),					TEXT("-targetplatform=Android -cookflavor=DXT"),		TEXT("Android"),	EPlatformSDKStatus::Unknown,	TEXT("/Engine/Tutorial/Mobile/SettingUpAndroidTutorial.SettingUpAndroidTutorial"),						IsAvailableOnWindows,						true),
+	BuildPlatformInfo(TEXT("Android_ETC1"),				TEXT("Android_ETC1"),		LOCTEXT("Android_ETC1", "Android (ETC1)"),					EPlatformType::Game,		EPlatformFlags::CookFlavor,		FPlatformIconPaths(TEXT("Launcher/AndroidTarget/Platform_Android_ETC1_24x"), TEXT("Launcher/AndroidTarget/Platform_Android_128x")),					TEXT("-targetplatform=Android -cookflavor=ETC1"),		TEXT("Android"),	EPlatformSDKStatus::Unknown,	TEXT("/Engine/Tutorial/Mobile/SettingUpAndroidTutorial.SettingUpAndroidTutorial"),						IsAvailableOnWindows,						true),
+	BuildPlatformInfo(TEXT("Android_ETC2"),				TEXT("Android_ETC2"),		LOCTEXT("Android_ETC2", "Android (ETC2)"),					EPlatformType::Game,		EPlatformFlags::CookFlavor,		FPlatformIconPaths(TEXT("Launcher/AndroidTarget/Platform_Android_ETC2_24x"), TEXT("Launcher/AndroidTarget/Platform_Android_128x")),					TEXT("-targetplatform=Android -cookflavor=ETC2"),		TEXT("Android"),	EPlatformSDKStatus::Unknown,	TEXT("/Engine/Tutorial/Mobile/SettingUpAndroidTutorial.SettingUpAndroidTutorial"),						IsAvailableOnWindows,						true),
+	BuildPlatformInfo(TEXT("Android_PVRTC"),			TEXT("Android_PVRTC"),		LOCTEXT("Android_PVRTC", "Android (PVRTC)"),				EPlatformType::Game,		EPlatformFlags::CookFlavor,		FPlatformIconPaths(TEXT("Launcher/AndroidTarget/Platform_Android_PVRTC_24x"), TEXT("Launcher/AndroidTarget/Platform_Android_128x")),				TEXT("-targetplatform=Android -cookflavor=PVRTC"),		TEXT("Android"),	EPlatformSDKStatus::Unknown,	TEXT("/Engine/Tutorial/Mobile/SettingUpAndroidTutorial.SettingUpAndroidTutorial"),						IsAvailableOnWindows,						true),
+	BuildPlatformInfo(TEXT("Android_ES31"),				TEXT("Android_ES31"),		LOCTEXT("Android_ES31", "Android (ES31)"),					EPlatformType::Game,		EPlatformFlags::CookFlavor,		FPlatformIconPaths(TEXT("Launcher/AndroidTarget/Platform_Android_ES31_24x"), TEXT("Launcher/AndroidTarget/Platform_Android_128x")),					TEXT("-targetplatform=Android -cookflavor=ES31"),		TEXT("Android"),	EPlatformSDKStatus::Unknown,	TEXT("Shared/Tutorials/SettingUpAndroidTutorial"),														IsAvailableOnWindows,						true),
+//	BuildPlatformInfo(TEXT("Android_GL4"),				TEXT("Android_GL4"),		LOCTEXT("Android_GL4", "Android (GL4)"),					EPlatformType::Game,		EPlatformFlags::CookFlavor,		FPlatformIconPaths(TEXT("Launcher/AndroidTarget/Platform_Android_GL4_24x"), TEXT("Launcher/AndroidTarget/Platform_Android_128x")),					TEXT("-targetplatform=Android -cookflavor=GL4"),		TEXT("Android"),	EPlatformSDKStatus::Unknown,	TEXT("Shared/Tutorials/SettingUpAndroidTutorial"),														IsAvailableOnWindows,						true),
+
+	BuildPlatformInfo(TEXT("HTML5"),					TEXT("HTML5"),				LOCTEXT("HTML5", "HTML5"),									EPlatformType::Game,		EPlatformFlags::None,			FPlatformIconPaths(TEXT("Launcher/HTML5Target/Platform_HTML5_24x"), TEXT("Launcher/HTML5Target/Platform_HTML5_128x")),								TEXT(""),												TEXT("HTML5"),		EPlatformSDKStatus::Unknown,	TEXT("Platforms/HTML5/GettingStarted"),																	IsAvailableOnWindows || IsAvailableOnMac,   false),
+
+	BuildPlatformInfo(TEXT("PS4"),						TEXT("PS4"),				LOCTEXT("PS4", "PlayStation 4"),							EPlatformType::Game,		EPlatformFlags::None,			FPlatformIconPaths(TEXT("Launcher/PS4Target/Platform_PS4_24x"), TEXT("Launcher/PS4Target/Platform_PS4_128x")),										TEXT(""),												TEXT("PS4"),		EPlatformSDKStatus::Unknown,	TEXT("Platforms/PS4/GettingStarted"),																	IsAvailableOnWindows,						false),
+
+	BuildPlatformInfo(TEXT("XboxOne"),					TEXT("XboxOne"),			LOCTEXT("XboxOne", "Xbox One"),								EPlatformType::Game,		EPlatformFlags::None,			FPlatformIconPaths(TEXT("Launcher/XboxOneTarget/Platform_XboxOne_24x"), TEXT("Launcher/XboxOneTarget/Platform_XboxOne_128x")),						TEXT(""),												TEXT("XboxOne"),	EPlatformSDKStatus::Unknown,	TEXT("Platforms/XboxOne/GettingStarted"),																IsAvailableOnWindows,						false),
 
 	//BuildPlatformInfo(TEXT("WinRT"),					TEXT("WinRT"),				LOCTEXT("WinRT", "Windows RT"),								EPlatformType::Game,		EPlatformFlags::None,			FPlatformIconPaths(TEXT("Launcher/WindowsTarget/Platform_WindowsNoEditor_24x"), TEXT("Launcher/WindowsTarget/Platform_WindowsNoEditor_128x")),		TEXT("")),
 };
diff --git a/Engine/Source/Editor/IntroTutorials/Private/IntroTutorials.cpp b/Engine/Source/Editor/IntroTutorials/Private/IntroTutorials.cpp
--- a/Engine/Source/Editor/IntroTutorials/Private/IntroTutorials.cpp
+++ b/Engine/Source/Editor/IntroTutorials/Private/IntroTutorials.cpp
@@ -31,6 +31,7 @@
 #include "ToolkitManager.h"
 #include "BlueprintEditor.h"
 #include "EngineBuildSettings.h"
+#include "IDocumentation.h"
 
 #define LOCTEXT_NAMESPACE "IntroTutorials"
 
@@ -341,7 +342,15 @@ void FIntroTutorials::HandleCompilerNotFound()
 
 void FIntroTutorials::HandleSDKNotInstalled(const FString& PlatformName, const FString& InTutorialAsset)
 {
-	LaunchTutorialByName( InTutorialAsset );
+	UBlueprint* Blueprint = LoadObject<UBlueprint>(nullptr, *InTutorialAsset);
+	if(Blueprint)
+	{
+		LaunchTutorialByName( InTutorialAsset );
+	}
+	else
+	{
+		IDocumentation::Get()->Open( InTutorialAsset );
+	}
 }
 
 EVisibility FIntroTutorials::GetHomeButtonVisibility() const
