Json: Fixed log category not being exported

diff --git a/Engine/Source/Developer/DesktopPlatform/Private/DesktopPlatformBase.cpp b/Engine/Source/Developer/DesktopPlatform/Private/DesktopPlatformBase.cpp
--- a/Engine/Source/Developer/DesktopPlatform/Private/DesktopPlatformBase.cpp
+++ b/Engine/Source/Developer/DesktopPlatform/Private/DesktopPlatformBase.cpp
@@ -6,9 +6,12 @@
 #include "EngineVersion.h"
 #include "ModuleManager.h"
 #include "Runtime/Launch/Resources/Version.h"
+#include "Json.h"
+
 
 #define LOCTEXT_NAMESPACE "DesktopPlatform"
 
+
 FDesktopPlatformBase::FDesktopPlatformBase()
 {
 	LauncherInstallationTimestamp = FDateTime::MinValue();
diff --git a/Engine/Source/Runtime/Json/Public/Json.h b/Engine/Source/Runtime/Json/Public/Json.h
--- a/Engine/Source/Runtime/Json/Public/Json.h
+++ b/Engine/Source/Runtime/Json/Public/Json.h
@@ -12,7 +12,8 @@
 /* Public includes
  *****************************************************************************/
 
-DECLARE_LOG_CATEGORY_EXTERN(LogJson, Log, All);
+JSON_API DECLARE_LOG_CATEGORY_EXTERN(LogJson, Log, All);
+
 
 #include "JsonPrintPolicy.h"
 #include "PrettyJsonPrintPolicy.h"
