Fix OS X system modal dialog input by setting a flag in MacApplication to indicate that a system modal dialog is active. We don't want to try handling any keyboard input in a system modal loop (the IME run-loop mode won't ever fire) and we need to allow blocks dispatched to the main thread to be executed in the NSModalPanelRunLoopMode.

diff --git a/Engine/Source/Developer/DesktopPlatform/Private/Mac/DesktopPlatformMac.cpp b/Engine/Source/Developer/DesktopPlatform/Private/Mac/DesktopPlatformMac.cpp
--- a/Engine/Source/Developer/DesktopPlatform/Private/Mac/DesktopPlatformMac.cpp
+++ b/Engine/Source/Developer/DesktopPlatform/Private/Mac/DesktopPlatformMac.cpp
@@ -8,6 +8,22 @@
 
 #define LOCTEXT_NAMESPACE "DesktopPlatform"
 
+class FMacScopedSystemModalMode
+{
+public:
+	FMacScopedSystemModalMode()
+	{
+		MacApplication->SystemModalMode(true);
+	}
+	
+	~FMacScopedSystemModalMode()
+	{
+		MacApplication->SystemModalMode(false);
+	}
+private:
+	FScopedSystemModalMode SystemModalMode;
+};
+
 class FCocoaScopeContext
 {
 public:
@@ -228,7 +244,7 @@ bool FDesktopPlatformMac::OpenDirectoryDialog(const void* ParentWindowHandle, co
 
 	bool bSuccess = false;
 	{
-		FScopedSystemModalMode SystemModalScope;
+		FMacScopedSystemModalMode SystemModalScope;
 		bSuccess = MainThreadReturn(^{
 			SCOPED_AUTORELEASE_POOL;
 			FCocoaScopeContext ContextGuard;
@@ -279,7 +295,7 @@ bool FDesktopPlatformMac::OpenFontDialog(const void* ParentWindowHandle, FString
 	
 	bool bSuccess = false;
 	{
-		FScopedSystemModalMode SystemModalScope;
+		FMacScopedSystemModalMode SystemModalScope;
 		bSuccess = MainThreadReturn(^{
 			SCOPED_AUTORELEASE_POOL;
 			FCocoaScopeContext ContextGuard;
@@ -411,7 +427,7 @@ bool FDesktopPlatformMac::FileDialogShared(bool bSave, const void* ParentWindowH
 
 	bool bSuccess = false;
 	{
-		FScopedSystemModalMode SystemModalScope;
+		FMacScopedSystemModalMode SystemModalScope;
 		bSuccess = MainThreadReturn(^{
 			SCOPED_AUTORELEASE_POOL;
 			FCocoaScopeContext ContextGuard;
diff --git a/Engine/Source/Runtime/Core/Private/Mac/CocoaThread.cpp b/Engine/Source/Runtime/Core/Private/Mac/CocoaThread.cpp
--- a/Engine/Source/Runtime/Core/Private/Mac/CocoaThread.cpp
+++ b/Engine/Source/Runtime/Core/Private/Mac/CocoaThread.cpp
@@ -138,6 +138,7 @@ private:
 		
 		// Register for default modes
 		Register(kCFRunLoopDefaultMode);
+		Register((CFStringRef)NSModalPanelRunLoopMode);
 		Register((CFStringRef)UE4NilEventMode);
 		Register((CFStringRef)UE4ShowEventMode);
 		Register((CFStringRef)UE4ResizeEventMode);
@@ -394,7 +395,7 @@ void MainThreadCall(dispatch_block_t Block, NSString* WaitMode, bool const bWait
 	if ( [NSThread mainThread] != [NSThread currentThread] )
 	{
 		FCocoaRunLoopSource& MainRunLoopSource = FCocoaRunLoopSource::GetMainRunLoopSource();
-		PerformBlockOnThread(MainRunLoopSource, Block, @[ NSDefaultRunLoopMode ], WaitMode, bWait);
+		PerformBlockOnThread(MainRunLoopSource, Block, @[ NSDefaultRunLoopMode, NSModalPanelRunLoopMode ], WaitMode, bWait);
 	}
 	else
 	{
diff --git a/Engine/Source/Runtime/Core/Private/Mac/MacApplication.cpp b/Engine/Source/Runtime/Core/Private/Mac/MacApplication.cpp
--- a/Engine/Source/Runtime/Core/Private/Mac/MacApplication.cpp
+++ b/Engine/Source/Runtime/Core/Private/Mac/MacApplication.cpp
@@ -109,6 +109,7 @@ FMacApplication::FMacApplication()
 	, MouseCaptureWindow( NULL )
 	, bIsMouseCaptureEnabled( false )
 	, bIsMouseCursorLocked( false )
+	, bSystemModalMode( false )
 	, ModifierKeysFlags( 0 )
 	, CurrentModifierFlags( 0 )
 	, EventMonitor( NULL )
@@ -734,7 +735,7 @@ void FMacApplication::ProcessNSEvent(NSEvent* const Event, TSharedPtr< FMacWindo
 		{
 			NSString *Characters = [Event characters];
 			bool bHandled = false;
-			if( [Characters length] && CurrentEventWindow.IsValid() )
+			if( !bSystemModalMode && [Characters length] && CurrentEventWindow.IsValid() )
 			{
 				bHandled = CurrentEventWindow->OnIMKKeyDown(Event);
 				if(!bHandled)
@@ -775,7 +776,7 @@ void FMacApplication::ProcessNSEvent(NSEvent* const Event, TSharedPtr< FMacWindo
 		{
 			NSString *Characters = [Event characters];
 			bool bHandled = false;
-			if( [Characters length] && CurrentEventWindow.IsValid() )
+			if( !bSystemModalMode && [Characters length] && CurrentEventWindow.IsValid() )
 			{
 				const bool IsRepeat = [Event isARepeat];
 				const TCHAR Character = ConvertChar( [Characters characterAtIndex:0] );
@@ -1422,3 +1423,8 @@ void FMacApplication::SendAnalytics(IAnalyticsProvider* Provider)
 }
 
 #endif
+
+void FMacApplication::SystemModalMode(bool const bInSystemModalMode)
+{
+	bSystemModalMode = bInSystemModalMode;
+}
diff --git a/Engine/Source/Runtime/Core/Public/Mac/MacApplication.h b/Engine/Source/Runtime/Core/Public/Mac/MacApplication.h
--- a/Engine/Source/Runtime/Core/Public/Mac/MacApplication.h
+++ b/Engine/Source/Runtime/Core/Public/Mac/MacApplication.h
@@ -85,6 +85,7 @@ public:
 	}
 #endif
 
+	void SystemModalMode(bool const bInSystemModalMode);
 
 public:
 
@@ -160,6 +161,8 @@ private:
 	bool bIsMouseCaptureEnabled;
 	bool bIsMouseCursorLocked;
 
+	bool bSystemModalMode;
+
 	TSharedPtr< FMacWindow > LastEventWindow;
 
 	/** The current set of modifier keys that are pressed. This is used to detect differences between left and right modifier keys on key up events*/
