- Removing ES31 Target Platform from being compiled - it won't work without the code to support it properly, which isn't quite there yet [ttp 347603]

diff --git a/Engine/Build/JunkManifest.txt b/Engine/Build/JunkManifest.txt
--- a/Engine/Build/JunkManifest.txt
+++ b/Engine/Build/JunkManifest.txt
@@ -49,3 +49,5 @@ Engine\Plugins\Script\ScriptPlugin\Binaries\Win32\*
 Engine\Plugins\Script\ScriptPlugin\Binaries\Win64\*
 Engine\Plugins\Experimental\Paper2D\Binaries
 Engine\Plugins\Experimental\Paper2D\Intermediate
+REM @todo es31: Take this out in for ES31 support
+Engine\Binaries\Win64\Android\*ES31*
diff --git a/Engine/Source/Developer/DesktopPlatform/Private/PlatformInfo.cpp b/Engine/Source/Developer/DesktopPlatform/Private/PlatformInfo.cpp
--- a/Engine/Source/Developer/DesktopPlatform/Private/PlatformInfo.cpp
+++ b/Engine/Source/Developer/DesktopPlatform/Private/PlatformInfo.cpp
@@ -101,7 +101,8 @@ static const FPlatformInfo PlatformInfoArray[] = {
 	BuildPlatformInfo(TEXT("Android_ETC1"),				TEXT("Android_ETC1"),		LOCTEXT("Android_ETC1", "Android (ETC1)"),					EPlatformType::Game,		EPlatformFlags::CookFlavor,		FPlatformIconPaths(TEXT("Launcher/AndroidTarget/Platform_Android_ETC1_24x"), TEXT("Launcher/AndroidTarget/Platform_Android_128x")),					TEXT("-targetplatform=Android -cookflavor=ETC1"),		TEXT("Android"),	EPlatformSDKStatus::Unknown,	TEXT("/Engine/Tutorial/Mobile/SettingUpAndroidTutorial.SettingUpAndroidTutorial"),						IsAvailableOnWindows,						true),
 	BuildPlatformInfo(TEXT("Android_ETC2"),				TEXT("Android_ETC2"),		LOCTEXT("Android_ETC2", "Android (ETC2)"),					EPlatformType::Game,		EPlatformFlags::CookFlavor,		FPlatformIconPaths(TEXT("Launcher/AndroidTarget/Platform_Android_ETC2_24x"), TEXT("Launcher/AndroidTarget/Platform_Android_128x")),					TEXT("-targetplatform=Android -cookflavor=ETC2"),		TEXT("Android"),	EPlatformSDKStatus::Unknown,	TEXT("/Engine/Tutorial/Mobile/SettingUpAndroidTutorial.SettingUpAndroidTutorial"),						IsAvailableOnWindows,						true),
 	BuildPlatformInfo(TEXT("Android_PVRTC"),			TEXT("Android_PVRTC"),		LOCTEXT("Android_PVRTC", "Android (PVRTC)"),				EPlatformType::Game,		EPlatformFlags::CookFlavor,		FPlatformIconPaths(TEXT("Launcher/AndroidTarget/Platform_Android_PVRTC_24x"), TEXT("Launcher/AndroidTarget/Platform_Android_128x")),				TEXT("-targetplatform=Android -cookflavor=PVRTC"),		TEXT("Android"),	EPlatformSDKStatus::Unknown,	TEXT("/Engine/Tutorial/Mobile/SettingUpAndroidTutorial.SettingUpAndroidTutorial"),						IsAvailableOnWindows,						true),
-	BuildPlatformInfo(TEXT("Android_ES31"),				TEXT("Android_ES31"),		LOCTEXT("Android_ES31", "Android (ES31)"),					EPlatformType::Game,		EPlatformFlags::CookFlavor,		FPlatformIconPaths(TEXT("Launcher/AndroidTarget/Platform_Android_ES31_24x"), TEXT("Launcher/AndroidTarget/Platform_Android_128x")),					TEXT("-targetplatform=Android -cookflavor=ES31"),		TEXT("Android"),	EPlatformSDKStatus::Unknown,	TEXT("Shared/Tutorials/SettingUpAndroidTutorial"),														IsAvailableOnWindows,						true),
+	// @todo es31: Put this back in for ES31 support
+//	BuildPlatformInfo(TEXT("Android_ES31"),				TEXT("Android_ES31"),		LOCTEXT("Android_ES31", "Android (ES31)"),					EPlatformType::Game,		EPlatformFlags::CookFlavor,		FPlatformIconPaths(TEXT("Launcher/AndroidTarget/Platform_Android_ES31_24x"), TEXT("Launcher/AndroidTarget/Platform_Android_128x")),					TEXT("-targetplatform=Android -cookflavor=ES31"),		TEXT("Android"),	EPlatformSDKStatus::Unknown,	TEXT("Shared/Tutorials/SettingUpAndroidTutorial"),														IsAvailableOnWindows,						true),
 //	BuildPlatformInfo(TEXT("Android_GL4"),				TEXT("Android_GL4"),		LOCTEXT("Android_GL4", "Android (GL4)"),					EPlatformType::Game,		EPlatformFlags::CookFlavor,		FPlatformIconPaths(TEXT("Launcher/AndroidTarget/Platform_Android_GL4_24x"), TEXT("Launcher/AndroidTarget/Platform_Android_128x")),					TEXT("-targetplatform=Android -cookflavor=GL4"),		TEXT("Android"),	EPlatformSDKStatus::Unknown,	TEXT("Shared/Tutorials/SettingUpAndroidTutorial"),														IsAvailableOnWindows,						true),
 
 	BuildPlatformInfo(TEXT("HTML5"),					TEXT("HTML5"),				LOCTEXT("HTML5", "HTML5"),									EPlatformType::Game,		EPlatformFlags::None,			FPlatformIconPaths(TEXT("Launcher/HTML5Target/Platform_HTML5_24x"), TEXT("Launcher/HTML5Target/Platform_HTML5_128x")),								TEXT(""),												TEXT("HTML5"),		EPlatformSDKStatus::Unknown,	TEXT("Platforms/HTML5/GettingStarted"),																	IsAvailableOnWindows || IsAvailableOnMac,   false),
diff --git a/Engine/Source/Programs/UnrealBuildTool/Android/UEBuildAndroid.cs b/Engine/Source/Programs/UnrealBuildTool/Android/UEBuildAndroid.cs
--- a/Engine/Source/Programs/UnrealBuildTool/Android/UEBuildAndroid.cs
+++ b/Engine/Source/Programs/UnrealBuildTool/Android/UEBuildAndroid.cs
@@ -237,7 +237,9 @@ namespace UnrealBuildTool
                             InModule.AddPlatformSpecificDynamicallyLoadedModule("Android_ETC1TargetPlatform");
                             InModule.AddPlatformSpecificDynamicallyLoadedModule("Android_ETC2TargetPlatform");
 //                            InModule.AddPlatformSpecificDynamicallyLoadedModule("Android_GL4TargetPlatform");
-                            InModule.AddPlatformSpecificDynamicallyLoadedModule("Android_ES31TargetPlatform");
+							// @todo es31: Put this back in for ES31 support
+							// ES31 code is in, but it's not fully supported in UE4 4.5, so for now we need to disable the targetplatform as it will confuse people greatly
+							// InModule.AddPlatformSpecificDynamicallyLoadedModule("Android_ES31TargetPlatform");
                         }
 					}
 					else if (InModule.ToString() == "UnrealEd")
@@ -267,7 +269,9 @@ namespace UnrealBuildTool
                     InModule.AddPlatformSpecificDynamicallyLoadedModule("Android_ETC1TargetPlatform");
                     InModule.AddPlatformSpecificDynamicallyLoadedModule("Android_ETC2TargetPlatform");
 //                    InModule.AddPlatformSpecificDynamicallyLoadedModule("Android_GL4TargetPlatform");
-                    InModule.AddPlatformSpecificDynamicallyLoadedModule("Android_ES31TargetPlatform");
+					// @todo es31: Put this back in for ES31 support
+					// ES31 code is in, but it's not fully supported in UE4 4.5, so for now we need to disable the targetplatform as it will confuse people greatly
+					// InModule.AddPlatformSpecificDynamicallyLoadedModule("Android_ES31TargetPlatform");
                 }
 
 				if (bBuildShaderFormats)
