Don't activate the Launcher on Mac when the editor starts if it's already running, but hidden

diff --git a/Engine/Source/Developer/DesktopPlatform/Private/Mac/DesktopPlatformMac.cpp b/Engine/Source/Developer/DesktopPlatform/Private/Mac/DesktopPlatformMac.cpp
--- a/Engine/Source/Developer/DesktopPlatform/Private/Mac/DesktopPlatformMac.cpp
+++ b/Engine/Source/Developer/DesktopPlatform/Private/Mac/DesktopPlatformMac.cpp
@@ -394,8 +394,11 @@ bool FDesktopPlatformMac::OpenLauncher(bool Install, FString CommandLineParams )
 	if ([RunningLaunchers count] > 0)
 	{
 		NSRunningApplication* Launcher = [RunningLaunchers objectAtIndex: 0];
-		[Launcher activateWithOptions: NSApplicationActivateAllWindows | NSApplicationActivateIgnoringOtherApps];
-		OpenLauncherCommandLine(CommandLineParams); // create a temp file that will tell running Launcher instance to switch to Marketplace tab
+		if (!Launcher.hidden || Install || CommandLineParams.Len() > 0) // If the launcher is running, but hidden, don't activate on editor startup
+		{
+			[Launcher activateWithOptions: NSApplicationActivateAllWindows | NSApplicationActivateIgnoringOtherApps];
+			OpenLauncherCommandLine(CommandLineParams); // create a temp file that will tell running Launcher instance to switch to Marketplace tab
+		}
 		return true;
 	}
 
