Add a failsafe for deducing the target name when rebuilding a project. In situations when the project has been renamed (by the launcher, or manually), the target name will still be the original name.

diff --git a/Engine/Source/Developer/DesktopPlatform/Private/DesktopPlatformBase.cpp b/Engine/Source/Developer/DesktopPlatform/Private/DesktopPlatformBase.cpp
--- a/Engine/Source/Developer/DesktopPlatform/Private/DesktopPlatformBase.cpp
+++ b/Engine/Source/Developer/DesktopPlatform/Private/DesktopPlatformBase.cpp
@@ -369,8 +369,30 @@ bool FDesktopPlatformBase::CleanGameProject(const FString& ProjectDir, FString&
 
 bool FDesktopPlatformBase::CompileGameProject(const FString& RootDir, const FString& ProjectFileName, FFeedbackContext* Warn)
 {
-	// Get the target name
-	FString Arguments = FString::Printf(TEXT("%s %s %s"), *FPaths::GetBaseFilename(ProjectFileName), FModuleManager::Get().GetUBTConfiguration(), FPlatformMisc::GetUBTPlatform());
+	// Get the project directory
+	FString ProjectDir = FPaths::GetPath(ProjectFileName);
+
+	// Get the target name. By default it'll be the same as the project name, but that might not be the case if the project was renamed.
+	FString TargetName = FPaths::GetBaseFilename(ProjectFileName);
+	if(!FPaths::FileExists(ProjectDir / FString::Printf(TEXT("Source/%sEditor.Target.cs"), *TargetName)))
+	{
+		// Find all the target files
+		TArray<FString> TargetFiles;
+		IFileManager::Get().FindFilesRecursive(TargetFiles, *(ProjectDir / TEXT("Source")), TEXT("*.target.cs"), true, false, false);
+
+		// Try to find a target that's clearly meant to be the editor. If there isn't one, let UBT fail with a sensible message without trying to do anything else smart.
+		for(const FString TargetFile: TargetFiles)
+		{
+			if(TargetFile.EndsWith("Editor.Target.cs"))
+			{
+				TargetName = FPaths::GetBaseFilename(FPaths::GetBaseFilename(TargetFile));
+				break;
+			}
+		}
+	}
+
+	// Build the argument list
+	FString Arguments = FString::Printf(TEXT("%s %s %s"), *TargetName, FModuleManager::Get().GetUBTConfiguration(), FPlatformMisc::GetUBTPlatform());
 
 	// Append the project name if it's a foreign project
 	if ( !ProjectFileName.IsEmpty() )
