PR-559: Changes needed for the 4.5.1.0 release of the Substance plugin

diff --git a/Engine/Source/Developer/DesktopPlatform/Private/Mac/DesktopPlatformMac.cpp b/Engine/Source/Developer/DesktopPlatform/Private/Mac/DesktopPlatformMac.cpp
--- a/Engine/Source/Developer/DesktopPlatform/Private/Mac/DesktopPlatformMac.cpp
+++ b/Engine/Source/Developer/DesktopPlatform/Private/Mac/DesktopPlatformMac.cpp
@@ -64,12 +64,14 @@ private:
 	NSTextField*	TextField;
 	NSSavePanel*	DialogPanel;
 	NSMutableArray*	AllowedFileTypes;
+	int32			SelectedExtension;
 }
 
 - (id)initWithFrame:(NSRect)frameRect dialogPanel:(NSSavePanel*) panel;
 - (void)PopUpButtonAction: (id) sender;
 - (void)AddAllowedFileTypes: (NSArray*) array;
 - (void)SetExtensionsAtIndex: (int32) index;
+- (int32)SelectedExtension;
 
 @end
 
@@ -135,6 +137,7 @@ private:
 - (void)SetExtensionsAtIndex: (int32) index
 {
 	check( [AllowedFileTypes count] >= index * 2 );
+	SelectedExtension = index;
 
 	NSString* ExtsToParse = [AllowedFileTypes objectAtIndex:index * 2 + 1];
 	if( [ExtsToParse compare:@"*.*"] == NSOrderedSame )
@@ -156,6 +159,10 @@ private:
 	}
 }
 
+- (int32)SelectedExtension {
+    return SelectedExtension;
+}
+
 @end
 
 /**
@@ -225,17 +232,19 @@ private:
 
 bool FDesktopPlatformMac::OpenFileDialog(const void* ParentWindowHandle, const FString& DialogTitle, const FString& DefaultPath, const FString& DefaultFile, const FString& FileTypes, uint32 Flags, TArray<FString>& OutFilenames)
 {
-	return FileDialogShared(false, ParentWindowHandle, DialogTitle, DefaultPath, DefaultFile, FileTypes, Flags, OutFilenames);
+	int32 DummyIdx = 0;
+	return FileDialogShared(false, ParentWindowHandle, DialogTitle, DefaultPath, DefaultFile, FileTypes, Flags, OutFilenames, DummyIdx);
 }
 
 bool FDesktopPlatformMac::OpenFileDialog(const void* ParentWindowHandle, const FString& DialogTitle, const FString& DefaultPath, const FString& DefaultFile, const FString& FileTypes, uint32 Flags, TArray<FString>& OutFilenames, int32& OutFilterIndex)
 {
-	return FileDialogShared(false, ParentWindowHandle, DialogTitle, DefaultPath, DefaultFile, FileTypes, Flags, OutFilenames);
+	return FileDialogShared(false, ParentWindowHandle, DialogTitle, DefaultPath, DefaultFile, FileTypes, Flags, OutFilenames, OutFilterIndex);
 }
 
 bool FDesktopPlatformMac::SaveFileDialog(const void* ParentWindowHandle, const FString& DialogTitle, const FString& DefaultPath, const FString& DefaultFile, const FString& FileTypes, uint32 Flags, TArray<FString>& OutFilenames)
 {
-	return FileDialogShared(true, ParentWindowHandle, DialogTitle, DefaultPath, DefaultFile, FileTypes, Flags, OutFilenames);
+	int32 DummyIdx = 0;
+	return FileDialogShared(true, ParentWindowHandle, DialogTitle, DefaultPath, DefaultFile, FileTypes, Flags, OutFilenames, DummyIdx);
 }
 
 bool FDesktopPlatformMac::OpenDirectoryDialog(const void* ParentWindowHandle, const FString& DialogTitle, const FString& DefaultPath, FString& OutFolderName)
@@ -424,7 +433,7 @@ bool FDesktopPlatformMac::OpenLauncher(bool Install, FString CommandLineParams )
 	return false;
 }
 
-bool FDesktopPlatformMac::FileDialogShared(bool bSave, const void* ParentWindowHandle, const FString& DialogTitle, const FString& DefaultPath, const FString& DefaultFile, const FString& FileTypes, uint32 Flags, TArray<FString>& OutFilenames)
+bool FDesktopPlatformMac::FileDialogShared(bool bSave, const void* ParentWindowHandle, const FString& DialogTitle, const FString& DefaultPath, const FString& DefaultFile, const FString& FileTypes, uint32 Flags, TArray<FString>& OutFilenames, int32& OutFilterIndex)
 {
 	MacApplication->SetCapture( NULL );
 
@@ -503,6 +512,7 @@ bool FDesktopPlatformMac::FileDialogShared(bool bSave, const void* ParentWindowH
 						FPlatformString::CFStringToTCHAR((CFStringRef)[FileURL path], FilePath);
 						new(OutFilenames) FString(FilePath);
 					}
+					OutFilterIndex = [AccessoryView SelectedExtension];
 				}
 
 				// Make sure all filenames gathered have their paths normalized
diff --git a/Engine/Source/Developer/DesktopPlatform/Private/Mac/DesktopPlatformMac.h b/Engine/Source/Developer/DesktopPlatform/Private/Mac/DesktopPlatformMac.h
--- a/Engine/Source/Developer/DesktopPlatform/Private/Mac/DesktopPlatformMac.h
+++ b/Engine/Source/Developer/DesktopPlatform/Private/Mac/DesktopPlatformMac.h
@@ -30,7 +30,7 @@ public:
 private:
 	bool GetLauncherPath(FString& OutLauncherPath) const;
 	bool GetLauncherInstallerPath(FString& OutInstallerPath) const;
-	bool FileDialogShared(bool bSave, const void* ParentWindowHandle, const FString& DialogTitle, const FString& DefaultPath, const FString& DefaultFile, const FString& FileTypes, uint32 Flags, TArray<FString>& OutFilenames);
+	bool FileDialogShared(bool bSave, const void* ParentWindowHandle, const FString& DialogTitle, const FString& DefaultPath, const FString& DefaultFile, const FString& FileTypes, uint32 Flags, TArray<FString>& OutFilenames, int32& OutFilterIndex);
 };
 
 typedef FDesktopPlatformMac FDesktopPlatform;
diff --git a/Engine/Source/Editor/TextureEditor/Private/TextureEditorToolkit.cpp b/Engine/Source/Editor/TextureEditor/Private/TextureEditorToolkit.cpp
--- a/Engine/Source/Editor/TextureEditor/Private/TextureEditorToolkit.cpp
+++ b/Engine/Source/Editor/TextureEditor/Private/TextureEditorToolkit.cpp
@@ -975,7 +975,7 @@ bool FTextureEditorToolkit::HandleRedChannelActionIsChecked( ) const
 
 bool FTextureEditorToolkit::HandleReimportActionCanExecute( ) const
 {
-	if (Texture->IsA<ULightMapTexture2D>() || Texture->IsA<UShadowMapTexture2D>())
+	if (Texture->IsA<ULightMapTexture2D>() || Texture->IsA<UShadowMapTexture2D>() || Texture->IsA<UTexture2DDynamic>())
 	{
 		return false;
 	}
diff --git a/Engine/Source/Programs/AutomationTool/TempStorage.cs b/Engine/Source/Programs/AutomationTool/TempStorage.cs
--- a/Engine/Source/Programs/AutomationTool/TempStorage.cs
+++ b/Engine/Source/Programs/AutomationTool/TempStorage.cs
@@ -69,7 +69,9 @@ namespace AutomationTool
                         LogType = System.Diagnostics.TraceEventType.Information;
                     }
 
-                    if (!((Timestamp - Other.Timestamp).TotalSeconds < 1 && (Timestamp - Other.Timestamp).TotalSeconds > -1))
+                    // on FAT filesystems writetime has a two seconds resolution
+                    // cf. http://msdn.microsoft.com/en-us/library/windows/desktop/ms724290%28v=vs.85%29.aspx
+                    if (!((Timestamp - Other.Timestamp).TotalSeconds < 2 && (Timestamp - Other.Timestamp).TotalSeconds > -2))
                     {
                         CommandUtils.Log(LogType, "File date mismatch {0} {1} {2} {3}", Name, Timestamp.ToString(), Other.Name, Other.Timestamp.ToString());
                         bOk = bOkToBeDifferent || bSilentOkToBeDifferent;
