Build UnrealBuildTool if it doesn't exist when we need to generate project files.

diff --git a/Engine/Source/Developer/DesktopPlatform/Private/DesktopPlatformBase.cpp b/Engine/Source/Developer/DesktopPlatform/Private/DesktopPlatformBase.cpp
--- a/Engine/Source/Developer/DesktopPlatform/Private/DesktopPlatformBase.cpp
+++ b/Engine/Source/Developer/DesktopPlatform/Private/DesktopPlatformBase.cpp
@@ -452,9 +452,18 @@ bool FDesktopPlatformBase::GenerateProjectFiles(const FString& RootDir, const FS
 	}
 	Arguments += " -progress";
 
-	// Run UnrealBuildTool
+	// Compile UnrealBuildTool if it doesn't exist. This can happen if we're just copying source from somewhere.
+	bool bRes = true;
 	Warn->BeginSlowTask(LOCTEXT("GeneratingProjectFiles", "Generating project files..."), true, true);
-	bool bRes = RunUnrealBuildTool(LOCTEXT("GeneratingProjectFiles", "Generating project files..."), RootDir, Arguments, Warn);
+	if(!FPaths::FileExists(GetUnrealBuildToolExecutableFilename()))
+	{
+		Warn->StatusUpdate(0, 1, LOCTEXT("BuildingUBT", "Building UnrealBuildTool..."));
+		bRes = BuildUnrealBuildTool(*Warn);
+	}
+	if(bRes)
+	{
+		bRes = RunUnrealBuildTool(LOCTEXT("GeneratingProjectFiles", "Generating project files..."), RootDir, Arguments, Warn);
+	}
 	Warn->EndSlowTask();
 	return bRes;
 }
@@ -904,12 +913,6 @@ FString FDesktopPlatformBase::GetUnrealBuildToolSourceCodePath() const
 
 bool FDesktopPlatformBase::BuildUnrealBuildTool(FOutputDevice &Ar)
 {
-	if (FApp::IsEngineInstalled())
-	{
-		// We may not build UBT in rocket
-		return false;
-	}
-
 	FString CompilerExecutableFilename;
 	FString CmdLineParams;
 #if PLATFORM_WINDOWS
@@ -944,7 +947,7 @@ bool FDesktopPlatformBase::BuildUnrealBuildTool(FOutputDevice &Ar)
 
 		FString BatchFileContents;
 		BatchFileContents = FString::Printf(TEXT("call \"%s\"") LINE_TERMINATOR, *VCVarsBat);
-		BatchFileContents += FString::Printf(TEXT("msbuild /nologo /verbosity:quiet %s /property:Configuration=Development /property:Platform=AnyCPU"), *CsProjLocation);
+		BatchFileContents += FString::Printf(TEXT("msbuild /nologo /verbosity:quiet \"%s\" /property:Configuration=Development /property:Platform=AnyCPU"), *CsProjLocation);
 		FFileHelper::SaveStringToFile(BatchFileContents, *BuildBatchFile);
 
 		{
