Made the Mac editor aware of Launcher name change

diff --git a/Engine/Source/Developer/DesktopPlatform/Private/Mac/DesktopPlatformMac.cpp b/Engine/Source/Developer/DesktopPlatform/Private/Mac/DesktopPlatformMac.cpp
--- a/Engine/Source/Developer/DesktopPlatform/Private/Mac/DesktopPlatformMac.cpp
+++ b/Engine/Source/Developer/DesktopPlatform/Private/Mac/DesktopPlatformMac.cpp
@@ -586,7 +586,7 @@ void FDesktopPlatformMac::EnumerateEngineInstallations(TMap<FString, FString> &O
 	for (auto It : Section)
 	{
 		const FString& EngineDir = It.Value;
-		if (EngineDir.Contains("Unreal Engine.app/Contents/") || EngineDir.Contains("[REDACTED_USER_HOME]/UnrealEngine/Launcher") || !IFileManager::Get().DirectoryExists(*EngineDir))
+		if (EngineDir.Contains("Unreal Engine.app/Contents/") || EngineDir.Contains("Epic Games Launcher.app/Contents/") || EngineDir.Contains("[REDACTED_USER_HOME]/UnrealEngine/Launcher") || !IFileManager::Get().DirectoryExists(*EngineDir))
 		{
 			KeysToRemove.Add(It.Key);
 		}
@@ -606,7 +606,7 @@ void FDesktopPlatformMac::EnumerateEngineInstallations(TMap<FString, FString> &O
 			NSBundle* AppBundle = [NSBundle bundleWithURL:AppURL];
 			FString EngineDir = FString([[AppBundle bundlePath] stringByDeletingLastPathComponent]);
 			if (([[AppBundle bundleIdentifier] isEqualToString:@"com.epicgames.UE4Editor"] || [[AppBundle bundleIdentifier] isEqualToString:@"com.epicgames.UE4EditorServices"])
-				&& EngineDir.RemoveFromEnd(TEXT("/Engine/Binaries/Mac")) && !EngineDir.Contains("Unreal Engine.app/Contents/") && !EngineDir.Contains("[REDACTED_USER_HOME]/UnrealEngine/Launcher"))
+				&& EngineDir.RemoveFromEnd(TEXT("/Engine/Binaries/Mac")) && !EngineDir.Contains("Unreal Engine.app/Contents/") && !EngineDir.Contains("Epic Games Launcher.app/Contents/") && !EngineDir.Contains("[REDACTED_USER_HOME]/UnrealEngine/Launcher"))
 			{
 				FString EngineId;
 				const FName* Key = Section.FindKey(EngineDir);
@@ -700,7 +700,7 @@ bool FDesktopPlatformMac::GetLauncherPath(FString& OutLauncherPath) const
 	}
 	else
 	{
-		FString LauncherPath = TEXT("/Applications/Unreal Engine.app/Contents/MacOS/UnrealEngineLauncher-Mac-Shipping");
+		FString LauncherPath = TEXT("/Applications/Epic Games Launcher.app/Contents/MacOS/UnrealEngineLauncher-Mac-Shipping");
 		if(FPaths::FileExists(LauncherPath))
 		{
 			OutLauncherPath = LauncherPath;
@@ -710,7 +710,7 @@ bool FDesktopPlatformMac::GetLauncherPath(FString& OutLauncherPath) const
 
 	// Otherwise search for it...
 	NSWorkspace* Workspace = [NSWorkspace sharedWorkspace];
-	NSString* Path = [Workspace fullPathForApplication:@"Unreal Engine"];
+	NSString* Path = [Workspace fullPathForApplication:@"Epic Games Launcher"];
 	if( Path )
 	{
 		OutLauncherPath = FString(Path);
@@ -723,7 +723,7 @@ bool FDesktopPlatformMac::GetLauncherPath(FString& OutLauncherPath) const
 bool FDesktopPlatformMac::GetLauncherInstallerPath(FString& OutInstallerPath) const
 {
 	// Check if the installer exists
-	FString InstallerPath = FPaths::ConvertRelativePathToFull(FPaths::Combine(*FPaths::EngineDir(), TEXT("Extras/UnrealEngineLauncher/UnrealEngine.dmg")));
+	FString InstallerPath = FPaths::ConvertRelativePathToFull(FPaths::Combine(*FPaths::EngineDir(), TEXT("Extras/UnrealEngineLauncher/EpicGamesLauncher.dmg")));
 	if (FPaths::FileExists(InstallerPath))
 	{
 		OutInstallerPath = InstallerPath;
