Initial support for Android deployment on Mac

diff --git a/Engine/Source/Developer/Android/AndroidDeviceDetection/Private/AndroidDeviceDetectionModule.cpp b/Engine/Source/Developer/Android/AndroidDeviceDetection/Private/AndroidDeviceDetectionModule.cpp
--- a/Engine/Source/Developer/Android/AndroidDeviceDetection/Private/AndroidDeviceDetectionModule.cpp
+++ b/Engine/Source/Developer/Android/AndroidDeviceDetection/Private/AndroidDeviceDetectionModule.cpp
@@ -260,7 +260,11 @@ public:
 			return;
 		}
 
+#if PLATFORM_WINDOWS
 		FString ADBPath = FString::Printf(TEXT("%s\\platform-tools\\adb.exe"), AndroidDirectory);
+#else
+		FString ADBPath = FString::Printf(TEXT("%s/platform-tools/adb"), AndroidDirectory);
+#endif
 
 		// if it doesn't exist, no SDK installed, so don't bother creating a thread to look for devices
 		if (!FPaths::FileExists(*ADBPath))
diff --git a/Engine/Source/Developer/Android/AndroidTargetPlatform/Private/AndroidTargetDevice.inl b/Engine/Source/Developer/Android/AndroidTargetPlatform/Private/AndroidTargetDevice.inl
--- a/Engine/Source/Developer/Android/AndroidTargetPlatform/Private/AndroidTargetDevice.inl
+++ b/Engine/Source/Developer/Android/AndroidTargetPlatform/Private/AndroidTargetDevice.inl
@@ -129,7 +129,11 @@ inline bool FAndroidTargetDevice::ExecuteAdbCommand( const FString& CommandLine,
 		return false;
 	}
 
+#if PLATFORM_WINDOWS
 	FString Filename = FString::Printf(TEXT("%s\\platform-tools\\adb.exe"), AndroidDirectory);
+#else
+	FString Filename = FString::Printf(TEXT("%s/platform-tools/adb"), AndroidDirectory);
+#endif
 
 	// execute the command
 	int32 ReturnCode;
diff --git a/Engine/Source/Developer/Android/AndroidTargetPlatform/Private/AndroidTargetPlatform.inl b/Engine/Source/Developer/Android/AndroidTargetPlatform/Private/AndroidTargetPlatform.inl
--- a/Engine/Source/Developer/Android/AndroidTargetPlatform/Private/AndroidTargetPlatform.inl
+++ b/Engine/Source/Developer/Android/AndroidTargetPlatform/Private/AndroidTargetPlatform.inl
@@ -101,18 +101,24 @@ inline bool FAndroidTargetPlatform<TPlatformProperties>::IsSdkInstalled(bool bPr
 	FPlatformMisc::GetEnvironmentVariable(TEXT("NDKROOT"), NDKROOT, MAX_PATH);
 
 	// make sure ANDROID_HOME points to the right thing
-	if (ANDROID_HOME[0] == 0 || 
+	if (ANDROID_HOME[0] == 0 ||
+#if PLATFORM_WINDOWS
 		IFileManager::Get().FileSize(*(FString(ANDROID_HOME) / TEXT("platform-tools/adb.exe"))) < 0)
+#else
+		IFileManager::Get().FileSize(*(FString(ANDROID_HOME) / TEXT("platform-tools/adb"))) < 0)
+#endif
 	{
 		return false;
 	}
 
+#if PLATFORM_WINDOWS
 	// make sure that JAVA_HOME points to the right thing
 	if (JAVA_HOME[0] == 0 ||
 		IFileManager::Get().FileSize(*(FString(JAVA_HOME) / TEXT("bin/javac.exe"))) < 0)
 	{
 		return false;
 	}
+#endif
 
 	// now look for ANT_HOME, or the ADT workaround of looking for a plugin
 	if (ANT_HOME[0] == 0)
diff --git a/Engine/Source/Developer/DesktopPlatform/Private/PlatformInfo.cpp b/Engine/Source/Developer/DesktopPlatform/Private/PlatformInfo.cpp
--- a/Engine/Source/Developer/DesktopPlatform/Private/PlatformInfo.cpp
+++ b/Engine/Source/Developer/DesktopPlatform/Private/PlatformInfo.cpp
@@ -94,13 +94,13 @@ static const FPlatformInfo PlatformInfoArray[] = {
 
 	BuildPlatformInfo(TEXT("IOS"),						TEXT("IOS"),				LOCTEXT("IOS", "iOS"),										EPlatformType::Game,		EPlatformFlags::None,			FPlatformIconPaths(TEXT("Launcher/iOSTarget/Platform_iOS_24x"), TEXT("Launcher/iOSTarget/Platform_iOS_128x")),										TEXT(""),												TEXT(""),			EPlatformSDKStatus::Unknown,	IOSTutorial,																							IsAvailableOnWindows || IsAvailableOnMac,	true),
 
-	BuildPlatformInfo(TEXT("Android"),					TEXT("Android"),			LOCTEXT("Android", "Android"),								EPlatformType::Game,		EPlatformFlags::None,			FPlatformIconPaths(TEXT("Launcher/AndroidTarget/Platform_Android_24x"), TEXT("Launcher/AndroidTarget/Platform_Android_128x")),						TEXT(""),												TEXT("Android"),	EPlatformSDKStatus::Unknown,	TEXT("/Engine/Tutorial/Mobile/SettingUpAndroidTutorial.SettingUpAndroidTutorial"),						IsAvailableOnWindows,						true),
-	BuildPlatformInfo(TEXT("Android_All"),				TEXT("Android"),			LOCTEXT("Android_All", "Android (All)"),					EPlatformType::Game,		EPlatformFlags::CookFlavor,		FPlatformIconPaths(TEXT("Launcher/AndroidTarget/Platform_Android_24x"), TEXT("Launcher/AndroidTarget/Platform_Android_128x")),						TEXT(""),												TEXT("Android"),	EPlatformSDKStatus::Unknown,	TEXT("/Engine/Tutorial/Mobile/SettingUpAndroidTutorial.SettingUpAndroidTutorial"),						IsAvailableOnWindows,						true),
-	BuildPlatformInfo(TEXT("Android_ATC"),				TEXT("Android_ATC"),		LOCTEXT("Android_ATC", "Android (ATC)"),					EPlatformType::Game,		EPlatformFlags::CookFlavor,		FPlatformIconPaths(TEXT("Launcher/AndroidTarget/Platform_Android_ATC_24x"), TEXT("Launcher/AndroidTarget/Platform_Android_128x")),					TEXT("-targetplatform=Android -cookflavor=ATC"),		TEXT("Android"),	EPlatformSDKStatus::Unknown,	TEXT("/Engine/Tutorial/Mobile/SettingUpAndroidTutorial.SettingUpAndroidTutorial"),						IsAvailableOnWindows,						true),
-	BuildPlatformInfo(TEXT("Android_DXT"),				TEXT("Android_DXT"),		LOCTEXT("Android_DXT", "Android (DXT)"),					EPlatformType::Game,		EPlatformFlags::CookFlavor,		FPlatformIconPaths(TEXT("Launcher/AndroidTarget/Platform_Android_DXT_24x"), TEXT("Launcher/AndroidTarget/Platform_Android_128x")),					TEXT("-targetplatform=Android -cookflavor=DXT"),		TEXT("Android"),	EPlatformSDKStatus::Unknown,	TEXT("/Engine/Tutorial/Mobile/SettingUpAndroidTutorial.SettingUpAndroidTutorial"),						IsAvailableOnWindows,						true),
-	BuildPlatformInfo(TEXT("Android_ETC1"),				TEXT("Android_ETC1"),		LOCTEXT("Android_ETC1", "Android (ETC1)"),					EPlatformType::Game,		EPlatformFlags::CookFlavor,		FPlatformIconPaths(TEXT("Launcher/AndroidTarget/Platform_Android_ETC1_24x"), TEXT("Launcher/AndroidTarget/Platform_Android_128x")),					TEXT("-targetplatform=Android -cookflavor=ETC1"),		TEXT("Android"),	EPlatformSDKStatus::Unknown,	TEXT("/Engine/Tutorial/Mobile/SettingUpAndroidTutorial.SettingUpAndroidTutorial"),						IsAvailableOnWindows,						true),
-	BuildPlatformInfo(TEXT("Android_ETC2"),				TEXT("Android_ETC2"),		LOCTEXT("Android_ETC2", "Android (ETC2)"),					EPlatformType::Game,		EPlatformFlags::CookFlavor,		FPlatformIconPaths(TEXT("Launcher/AndroidTarget/Platform_Android_ETC2_24x"), TEXT("Launcher/AndroidTarget/Platform_Android_128x")),					TEXT("-targetplatform=Android -cookflavor=ETC2"),		TEXT("Android"),	EPlatformSDKStatus::Unknown,	TEXT("/Engine/Tutorial/Mobile/SettingUpAndroidTutorial.SettingUpAndroidTutorial"),						IsAvailableOnWindows,						true),
-	BuildPlatformInfo(TEXT("Android_PVRTC"),			TEXT("Android_PVRTC"),		LOCTEXT("Android_PVRTC", "Android (PVRTC)"),				EPlatformType::Game,		EPlatformFlags::CookFlavor,		FPlatformIconPaths(TEXT("Launcher/AndroidTarget/Platform_Android_PVRTC_24x"), TEXT("Launcher/AndroidTarget/Platform_Android_128x")),				TEXT("-targetplatform=Android -cookflavor=PVRTC"),		TEXT("Android"),	EPlatformSDKStatus::Unknown,	TEXT("/Engine/Tutorial/Mobile/SettingUpAndroidTutorial.SettingUpAndroidTutorial"),						IsAvailableOnWindows,						true),
+	BuildPlatformInfo(TEXT("Android"),					TEXT("Android"),			LOCTEXT("Android", "Android"),								EPlatformType::Game,		EPlatformFlags::None,			FPlatformIconPaths(TEXT("Launcher/AndroidTarget/Platform_Android_24x"), TEXT("Launcher/AndroidTarget/Platform_Android_128x")),						TEXT(""),												TEXT("Android"),	EPlatformSDKStatus::Unknown,	TEXT("/Engine/Tutorial/Mobile/SettingUpAndroidTutorial.SettingUpAndroidTutorial"),						IsAvailableOnWindows || IsAvailableOnMac,	true),
+	BuildPlatformInfo(TEXT("Android_All"),				TEXT("Android"),			LOCTEXT("Android_All", "Android (All)"),					EPlatformType::Game,		EPlatformFlags::CookFlavor,		FPlatformIconPaths(TEXT("Launcher/AndroidTarget/Platform_Android_24x"), TEXT("Launcher/AndroidTarget/Platform_Android_128x")),						TEXT(""),												TEXT("Android"),	EPlatformSDKStatus::Unknown,	TEXT("/Engine/Tutorial/Mobile/SettingUpAndroidTutorial.SettingUpAndroidTutorial"),						IsAvailableOnWindows || IsAvailableOnMac,	true),
+	BuildPlatformInfo(TEXT("Android_ATC"),				TEXT("Android_ATC"),		LOCTEXT("Android_ATC", "Android (ATC)"),					EPlatformType::Game,		EPlatformFlags::CookFlavor,		FPlatformIconPaths(TEXT("Launcher/AndroidTarget/Platform_Android_ATC_24x"), TEXT("Launcher/AndroidTarget/Platform_Android_128x")),					TEXT("-targetplatform=Android -cookflavor=ATC"),		TEXT("Android"),	EPlatformSDKStatus::Unknown,	TEXT("/Engine/Tutorial/Mobile/SettingUpAndroidTutorial.SettingUpAndroidTutorial"),						IsAvailableOnWindows || IsAvailableOnMac,	true),
+	BuildPlatformInfo(TEXT("Android_DXT"),				TEXT("Android_DXT"),		LOCTEXT("Android_DXT", "Android (DXT)"),					EPlatformType::Game,		EPlatformFlags::CookFlavor,		FPlatformIconPaths(TEXT("Launcher/AndroidTarget/Platform_Android_DXT_24x"), TEXT("Launcher/AndroidTarget/Platform_Android_128x")),					TEXT("-targetplatform=Android -cookflavor=DXT"),		TEXT("Android"),	EPlatformSDKStatus::Unknown,	TEXT("/Engine/Tutorial/Mobile/SettingUpAndroidTutorial.SettingUpAndroidTutorial"),						IsAvailableOnWindows || IsAvailableOnMac,	true),
+	BuildPlatformInfo(TEXT("Android_ETC1"),				TEXT("Android_ETC1"),		LOCTEXT("Android_ETC1", "Android (ETC1)"),					EPlatformType::Game,		EPlatformFlags::CookFlavor,		FPlatformIconPaths(TEXT("Launcher/AndroidTarget/Platform_Android_ETC1_24x"), TEXT("Launcher/AndroidTarget/Platform_Android_128x")),					TEXT("-targetplatform=Android -cookflavor=ETC1"),		TEXT("Android"),	EPlatformSDKStatus::Unknown,	TEXT("/Engine/Tutorial/Mobile/SettingUpAndroidTutorial.SettingUpAndroidTutorial"),						IsAvailableOnWindows || IsAvailableOnMac,	true),
+	BuildPlatformInfo(TEXT("Android_ETC2"),				TEXT("Android_ETC2"),		LOCTEXT("Android_ETC2", "Android (ETC2)"),					EPlatformType::Game,		EPlatformFlags::CookFlavor,		FPlatformIconPaths(TEXT("Launcher/AndroidTarget/Platform_Android_ETC2_24x"), TEXT("Launcher/AndroidTarget/Platform_Android_128x")),					TEXT("-targetplatform=Android -cookflavor=ETC2"),		TEXT("Android"),	EPlatformSDKStatus::Unknown,	TEXT("/Engine/Tutorial/Mobile/SettingUpAndroidTutorial.SettingUpAndroidTutorial"),						IsAvailableOnWindows || IsAvailableOnMac,	true),
+	BuildPlatformInfo(TEXT("Android_PVRTC"),			TEXT("Android_PVRTC"),		LOCTEXT("Android_PVRTC", "Android (PVRTC)"),				EPlatformType::Game,		EPlatformFlags::CookFlavor,		FPlatformIconPaths(TEXT("Launcher/AndroidTarget/Platform_Android_PVRTC_24x"), TEXT("Launcher/AndroidTarget/Platform_Android_128x")),				TEXT("-targetplatform=Android -cookflavor=PVRTC"),		TEXT("Android"),	EPlatformSDKStatus::Unknown,	TEXT("/Engine/Tutorial/Mobile/SettingUpAndroidTutorial.SettingUpAndroidTutorial"),						IsAvailableOnWindows || IsAvailableOnMac,	true),
 	// @todo es31: Set the final false to true when ES31 is in Binary releases
 	BuildPlatformInfo(TEXT("Android_ES31"),				TEXT("Android_ES31"),		LOCTEXT("Android_ES31", "Android (ES31+AEP)"),				EPlatformType::Game,		EPlatformFlags::CookFlavor,		FPlatformIconPaths(TEXT("Launcher/AndroidTarget/Platform_Android_ES31_24x"), TEXT("Launcher/AndroidTarget/Platform_Android_128x")),					TEXT("-targetplatform=Android -cookflavor=ES31"),		TEXT("Android"),	EPlatformSDKStatus::Unknown,	TEXT("Shared/Tutorials/SettingUpAndroidTutorial"),														IsAvailableOnWindows,						false),
 //	BuildPlatformInfo(TEXT("Android_GL4"),				TEXT("Android_GL4"),		LOCTEXT("Android_GL4", "Android (GL4)"),					EPlatformType::Game,		EPlatformFlags::CookFlavor,		FPlatformIconPaths(TEXT("Launcher/AndroidTarget/Platform_Android_GL4_24x"), TEXT("Launcher/AndroidTarget/Platform_Android_128x")),					TEXT("-targetplatform=Android -cookflavor=GL4"),		TEXT("Android"),	EPlatformSDKStatus::Unknown,	TEXT("Shared/Tutorials/SettingUpAndroidTutorial"),														IsAvailableOnWindows,						true),
diff --git a/Engine/Source/Programs/AutomationTool/Android/AndroidPlatform.Automation.cs b/Engine/Source/Programs/AutomationTool/Android/AndroidPlatform.Automation.cs
--- a/Engine/Source/Programs/AutomationTool/Android/AndroidPlatform.Automation.cs
+++ b/Engine/Source/Programs/AutomationTool/Android/AndroidPlatform.Automation.cs
@@ -95,12 +95,19 @@ public class AndroidPlatform : Platform
 
     private static string GetStorageQueryCommand()
     {
-        return "shell \"echo $EXTERNAL_STORAGE\"";
+		if (Utils.IsRunningOnMono)
+		{
+			return "shell 'echo $EXTERNAL_STORAGE'";
+		}
+		else
+		{
+			return "shell \"echo $EXTERNAL_STORAGE\"";
+		}
     }
 
 	private static string GetFinalBatchName(string ApkName, ProjectParams Params, string Architecture, string GPUArchitecture)
 	{
-		return Path.Combine(Path.GetDirectoryName(ApkName), "Install_" + Params.ShortProjectName + "_" + Params.ClientConfigsToBuild[0].ToString() + Architecture + GPUArchitecture + ".bat");
+		return Path.Combine(Path.GetDirectoryName(ApkName), "Install_" + Params.ShortProjectName + "_" + Params.ClientConfigsToBuild[0].ToString() + Architecture + GPUArchitecture + (Utils.IsRunningOnMono ? ".sh" : ".bat"));
 	}
 
 	public override void Package(ProjectParams Params, DeploymentContext SC, int WorkingCL)
@@ -176,34 +183,68 @@ public class AndroidPlatform : Platform
 
 				// Write install batch file(s).
 
-				Log("Writing bat for install with {0}", Params.OBBinAPK ? "OBB in APK" : "OBB separate");
 				string PackageName = GetPackageInfo(ApkName, false);
 				// make a batch file that can be used to install the .apk and .obb files
-				string[] BatchLines = new string[] {
-					"setlocal",
-					"set ADB=%ANDROID_HOME%\\platform-tools\\adb.exe",
-					"set DEVICE=",
-					"if not \"%1\"==\"\" set DEVICE=-s %1",
-					"for /f \"delims=\" %%A in ('%ADB% %DEVICE% " + GetStorageQueryCommand() +"') do @set STORAGE=%%A",
-					"%ADB% %DEVICE% uninstall " + PackageName,
-					"%ADB% %DEVICE% install " + Path.GetFileName(ApkName),
-					"@if \"%ERRORLEVEL%\" NEQ \"0\" goto Error",
-					"%ADB% %DEVICE% shell rm -r %STORAGE%/" + Params.ShortProjectName,
-					"%ADB% %DEVICE% shell rm -r %STORAGE%/UE4Game/UE4CommandLine.txt", // we need to delete the commandline in UE4Game or it will mess up loading
-					"%ADB% %DEVICE% shell rm -r %STORAGE%/obb/" + PackageName,
-					Params.OBBinAPK ? "" : "%ADB% %DEVICE% push " + Path.GetFileName(ObbName) + " %STORAGE%/" + DeviceObbName,
-					Params.OBBinAPK ? "" : "if \"%ERRORLEVEL%\" NEQ \"0\" goto Error",
-					"goto:eof",
-					":Error",
-					"@echo.",
-					"@echo There was an error installing the game or the obb file. Look above for more info.",
-					"@echo.",
-					"@echo Things to try:",
-					"@echo Check that the device (and only the device) is listed with \"%ADB$ devices\" from a command prompt.",
-					"@echo Make sure all Developer options look normal on the device",
-					"@echo Check that the device has an SD card.",
-					"@pause"
-				};
+				string[] BatchLines;
+				if (Utils.IsRunningOnMono)
+				{
+					Log("Writing shell script for install with {0}", Params.OBBinAPK ? "OBB in APK" : "OBB separate");
+					BatchLines = new string[] {
+						"#!/bin/sh",
+						"ADB=$ANDROID_HOME/platform-tools/adb",
+						"DEVICE=",
+						"if [ \"$1\" != \"\" ]; then DEVICE=\"-s $1\"; fi",
+						"$ADB $DEVICE uninstall " + PackageName,
+						"$ADB $DEVICE install " + Path.GetFileName(ApkName),
+						"if [ $? -eq 0 ]; then",
+						"\t$ADB $DEVICE shell 'rm -r $EXTERNAL_STORAGE/" + Params.ShortProjectName + "'",
+						"\t$ADB $DEVICE shell 'rm -r $EXTERNAL_STORAGE/UE4Game/UE4CommandLine.txt" + "'",
+						"\t$ADB $DEVICE shell 'rm -r $EXTERNAL_STORAGE/obb/" + PackageName + "'",
+						"\tSTORAGE=$(echo \"`$ADB $DEVICE shell 'echo $EXTERNAL_STORAGE'`\" | cat -v | tr -d '^M')",
+						"\t$ADB $DEVICE push " + Path.GetFileName(ObbName) + " $STORAGE/" + DeviceObbName,
+						"\tif [ $? -eq 0 ]; then",
+						"\t\texit 0",
+						"\tfi",
+						"fi",
+						"echo",
+						"echo \"There was an error installing the game or the obb file. Look above for more info.\"",
+						"echo",
+						"echo \"Things to try:\"",
+						"echo \"Check that the device (and only the device) is listed with \\\"$ADB devices\\\" from a command prompt.\"",
+						"echo \"Make sure all Developer options look normal on the device\"",
+						"echo \"Check that the device has an SD card.\"",
+						"exit 1"
+					};
+				}
+				else
+				{
+					Log("Writing bat for install with {0}", Params.OBBinAPK ? "OBB in APK" : "OBB separate");
+					BatchLines = new string[] {
+						"setlocal",
+						"set ADB=%ANDROID_HOME%\\platform-tools\\adb.exe",
+						"set DEVICE=",
+						"if not \"%1\"==\"\" set DEVICE=-s %1",
+						"for /f \"delims=\" %%A in ('%ADB% %DEVICE% " + GetStorageQueryCommand() + "') do @set STORAGE=%%A",
+						"%ADB% %DEVICE% uninstall " + PackageName,
+						"%ADB% %DEVICE% install " + Path.GetFileName(ApkName),
+						"@if \"%ERRORLEVEL%\" NEQ \"0\" goto Error",
+						"%ADB% %DEVICE% shell rm -r %STORAGE%/" + Params.ShortProjectName,
+						"%ADB% %DEVICE% shell rm -r %STORAGE%/UE4Game/UE4CommandLine.txt", // we need to delete the commandline in UE4Game or it will mess up loading
+						"%ADB% %DEVICE% shell rm -r %STORAGE%/obb/" + PackageName,
+						Params.OBBinAPK ? "" : "%ADB% %DEVICE% push " + Path.GetFileName(ObbName) + " %STORAGE%/" + DeviceObbName,
+						Params.OBBinAPK ? "" : "if \"%ERRORLEVEL%\" NEQ \"0\" goto Error",
+						"goto:eof",
+						":Error",
+						"@echo.",
+						"@echo There was an error installing the game or the obb file. Look above for more info.",
+						"@echo.",
+						"@echo Things to try:",
+						"@echo Check that the device (and only the device) is listed with \"%ADB$ devices\" from a command prompt.",
+						"@echo Make sure all Developer options look normal on the device",
+						"@echo Check that the device has an SD card.",
+						"@pause"
+					};
+				}
 				File.WriteAllLines(BatchName, BatchLines);
 			}
 		}
diff --git a/Engine/Source/Programs/UnrealBuildTool/Android/AndroidToolChain.cs b/Engine/Source/Programs/UnrealBuildTool/Android/AndroidToolChain.cs
--- a/Engine/Source/Programs/UnrealBuildTool/Android/AndroidToolChain.cs
+++ b/Engine/Source/Programs/UnrealBuildTool/Android/AndroidToolChain.cs
@@ -134,7 +134,7 @@ namespace UnrealBuildTool
 			if (Target == "latest")
 			{
 				// get a list of NDK platforms
-				string PlatformsDir = Environment.ExpandEnvironmentVariables("%NDKROOT%\\platforms");
+				string PlatformsDir = Environment.ExpandEnvironmentVariables("%NDKROOT%/platforms");
 				if (!Directory.Exists(PlatformsDir))
 				{
 					throw new BuildException("No platforms found in {0}", PlatformsDir);
@@ -195,17 +195,17 @@ namespace UnrealBuildTool
 			string GccVersion = "";
 
 			// prefer clang 3.5, but fall back if needed for now
-			if (Directory.Exists(Path.Combine(NDKPath, @"toolchains\llvm-3.5")))
+			if (Directory.Exists(Path.Combine(NDKPath, @"toolchains/llvm-3.5")))
 			{
 				ClangVersion = "3.5";
 				GccVersion = "4.9";
 			}
-			else if (Directory.Exists(Path.Combine(NDKPath, @"toolchains\llvm-3.3")))
+			else if (Directory.Exists(Path.Combine(NDKPath, @"toolchains/llvm-3.3")))
 			{
 				ClangVersion = "3.3";
 				GccVersion = "4.8";
 			}
-			else if (Directory.Exists(Path.Combine(NDKPath, @"toolchains\llvm-3.1")))
+			else if (Directory.Exists(Path.Combine(NDKPath, @"toolchains/llvm-3.1")))
 			{
 				ClangVersion = "3.1";
 				GccVersion = "4.6";
@@ -219,8 +219,10 @@ namespace UnrealBuildTool
 			// Console.WriteLine("Compiling with clang {0}", ClangVersionFloat);
 
             string ArchitecturePath = "";
-            string ArchitecturePathWindows32 = @"prebuilt\windows";
-            string ArchitecturePathWindows64 = @"prebuilt\windows-x86_64";
+            string ArchitecturePathWindows32 = @"prebuilt/windows";
+            string ArchitecturePathWindows64 = @"prebuilt/windows-x86_64";
+			string ArchitecturePathMac = @"prebuilt/darwin-x86_64";
+			string ExeExtension = ".exe";
 
             if (Directory.Exists(Path.Combine(NDKPath, ArchitecturePathWindows64)))
             {
@@ -232,6 +234,12 @@ namespace UnrealBuildTool
                 Log.TraceVerbose("        Found Windows 32 bit versions of toolchain");
                 ArchitecturePath = ArchitecturePathWindows32;
             }
+			else if (Directory.Exists(Path.Combine(NDKPath, ArchitecturePathMac)))
+			{
+				Log.TraceVerbose("        Found Mac versions of toolchain");
+				ArchitecturePath = ArchitecturePathMac;
+				ExeExtension = "";
+			}
             else
             {
                 Log.TraceVerbose("        Did not find 32 bit or 64 bit versions of toolchain");
@@ -239,18 +247,18 @@ namespace UnrealBuildTool
             }
 
 			// set up the path to our toolchains
-			ClangPath = Path.Combine(NDKPath, @"toolchains\llvm-" + ClangVersion, ArchitecturePath, @"bin\clang++.exe");
-            ArPathArm = Path.Combine(NDKPath, @"toolchains\arm-linux-androideabi-" + GccVersion, ArchitecturePath, @"bin\arm-linux-androideabi-ar.exe");		//@todo android: use llvm-ar.exe instead?
-            ArPathx86 = Path.Combine(NDKPath, @"toolchains\x86-" + GccVersion, ArchitecturePath, @"bin\i686-linux-android-ar.exe");							//@todo android: verify x86 toolchain
+			ClangPath = Path.Combine(NDKPath, @"toolchains/llvm-" + ClangVersion, ArchitecturePath, @"bin/clang++" + ExeExtension);
+			ArPathArm = Path.Combine(NDKPath, @"toolchains/arm-linux-androideabi-" + GccVersion, ArchitecturePath, @"bin/arm-linux-androideabi-ar" + ExeExtension);		//@todo android: use llvm-ar.exe instead?
+			ArPathx86 = Path.Combine(NDKPath, @"toolchains/x86-" + GccVersion, ArchitecturePath, @"bin/i686-linux-android-ar" + ExeExtension);							//@todo android: verify x86 toolchain
 
 
 			// toolchain params
 			ToolchainParamsArm = " -target armv7-none-linux-androideabi" +
 								   " --sysroot=\"" + Path.Combine(NDKPath, "platforms", GetNdkApiLevel(), "arch-arm") + "\"" +
-                                   " -gcc-toolchain \"" + Path.Combine(NDKPath, @"toolchains\arm-linux-androideabi-" + GccVersion, ArchitecturePath) + "\"";
+                                   " -gcc-toolchain \"" + Path.Combine(NDKPath, @"toolchains/arm-linux-androideabi-" + GccVersion, ArchitecturePath) + "\"";
 			ToolchainParamsx86 = " -target i686-none-linux-android" +
 								   " --sysroot=\"" + Path.Combine(NDKPath, "platforms", GetNdkApiLevel(), "arch-x86") + "\"" +
-                                   " -gcc-toolchain \"" + Path.Combine(NDKPath, @"toolchains\x86-" + GccVersion, ArchitecturePath) + "\"";
+                                   " -gcc-toolchain \"" + Path.Combine(NDKPath, @"toolchains/x86-" + GccVersion, ArchitecturePath) + "\"";
 
 			// Register this tool chain
 			Log.TraceVerbose("        Registered for {0}", CPPTargetPlatform.Android.ToString());
diff --git a/Engine/Source/Programs/UnrealBuildTool/Android/UEBuildAndroid.cs b/Engine/Source/Programs/UnrealBuildTool/Android/UEBuildAndroid.cs
--- a/Engine/Source/Programs/UnrealBuildTool/Android/UEBuildAndroid.cs
+++ b/Engine/Source/Programs/UnrealBuildTool/Android/UEBuildAndroid.cs
@@ -74,9 +74,9 @@ namespace UnrealBuildTool
             NDKPath = NDKPath.Replace("\"", "");
 
             // need a supported llvm
-            if (!Directory.Exists(Path.Combine(NDKPath, @"toolchains\llvm-3.5")) && 
-				!Directory.Exists(Path.Combine(NDKPath, @"toolchains\llvm-3.3")) &&
-				!Directory.Exists(Path.Combine(NDKPath, @"toolchains\llvm-3.1")))
+            if (!Directory.Exists(Path.Combine(NDKPath, @"toolchains/llvm-3.5")) && 
+				!Directory.Exists(Path.Combine(NDKPath, @"toolchains/llvm-3.3")) &&
+				!Directory.Exists(Path.Combine(NDKPath, @"toolchains/llvm-3.1")))
             {
                 return false;
             }
@@ -285,7 +285,7 @@ namespace UnrealBuildTool
 			NDKPath = NDKPath.Replace("\"", "");
 
 			string GccVersion = "4.8";
-			if (!Directory.Exists(Path.Combine(NDKPath, @"sources\cxx-stl\gnu-libstdc++\4.8")))
+			if (!Directory.Exists(Path.Combine(NDKPath, @"sources/cxx-stl/gnu-libstdc++/4.8")))
 			{
 				GccVersion = "4.6";
 			}
diff --git a/Engine/Source/Programs/UnrealBuildTool/Android/UEDeployAndroid.cs b/Engine/Source/Programs/UnrealBuildTool/Android/UEDeployAndroid.cs
--- a/Engine/Source/Programs/UnrealBuildTool/Android/UEDeployAndroid.cs
+++ b/Engine/Source/Programs/UnrealBuildTool/Android/UEDeployAndroid.cs
@@ -72,7 +72,7 @@ namespace UnrealBuildTool.Android
 				if (Target == "latest")
 				{
 					// we expect there to be one, so use the first one
-					string AndroidCommandPath = Environment.ExpandEnvironmentVariables("%ANDROID_HOME%/tools/android.bat");
+					string AndroidCommandPath = Environment.ExpandEnvironmentVariables("%ANDROID_HOME%/tools/android" + (Utils.IsRunningOnMono ? "" : ".bat"));
 
 					var ExeInfo = new ProcessStartInfo(AndroidCommandPath, "list targets");
 					ExeInfo.UseShellExecute = false;
@@ -108,7 +108,7 @@ namespace UnrealBuildTool.Android
 			string AntHome = Environment.GetEnvironmentVariable("ANT_HOME");
 			if (!string.IsNullOrEmpty(AntHome) && Directory.Exists(AntHome))
 			{
-				string AntPath = AntHome + "/bin/ant.bat";
+				string AntPath = AntHome + "/bin/ant" + (Utils.IsRunningOnMono ? "" : ".bat");
 				// use it if found
 				if (File.Exists(AntPath))
 				{
@@ -126,7 +126,7 @@ namespace UnrealBuildTool.Android
 				{
 					foreach (string PluginName in Plugins)
 					{
-						string AntPath = PluginName + "/bin/ant.bat";
+						string AntPath = PluginName + "/bin/ant" + (Utils.IsRunningOnMono ? "" : ".bat");
 						// use it if found
 						if (File.Exists(AntPath))
 						{
@@ -440,7 +440,7 @@ namespace UnrealBuildTool.Android
 			}
 			
 			// now update the project for each library
-			string AndroidCommandPath = Environment.ExpandEnvironmentVariables("%ANDROID_HOME%/tools/android.bat");
+			string AndroidCommandPath = Environment.ExpandEnvironmentVariables("%ANDROID_HOME%/tools/android" + (Utils.IsRunningOnMono ? "" : ".bat"));
 			string UpdateCommandLine = "--silent update project --subprojects --name " + ProjectName + " --path . --target " + GetSdkApiLevel();
 			foreach (string Lib in LibsToBeAdded)
 			{
@@ -570,7 +570,7 @@ namespace UnrealBuildTool.Android
 			Log.TraceInformation("\n===={0}====PREPARING TO MAKE APK=================================================================", DateTime.Now.ToString());
 
 			// cache some tools paths
-			string NDKBuildPath = Environment.ExpandEnvironmentVariables("%NDKROOT%/ndk-build.cmd");
+			string NDKBuildPath = Environment.ExpandEnvironmentVariables("%NDKROOT%/ndk-build" + (Utils.IsRunningOnMono ? "" : ".cmd"));
 
 			// set up some directory info
 			string IntermediateAndroidPath = Path.Combine(ProjectDirectory, "Intermediate/Android/");
@@ -793,7 +793,14 @@ namespace UnrealBuildTool.Android
 						}
 
 						// Use ant to build the .apk file
-						RunCommandLineProgramAndThrowOnError(UE4BuildPath, "cmd.exe", "/c \"" + GetAntPath() + "\" -quiet " + AntBuildType, "Making .apk with Ant... (note: it's safe to ignore javac obsolete warnings)");
+						if (Utils.IsRunningOnMono)
+						{
+							RunCommandLineProgramAndThrowOnError(UE4BuildPath, "/bin/sh", "-c '\"" + GetAntPath() + "\" -quiet " + AntBuildType + "'", "Making .apk with Ant... (note: it's safe to ignore javac obsolete warnings)");
+						}
+						else
+						{
+							RunCommandLineProgramAndThrowOnError(UE4BuildPath, "cmd.exe", "/c \"" + GetAntPath() + "\" -quiet " + AntBuildType, "Making .apk with Ant... (note: it's safe to ignore javac obsolete warnings)");
+						}
 
 						// make sure destination exists
 						Directory.CreateDirectory(Path.GetDirectoryName(DestApkName));
