Made the editor able to detect both old and new versions of the launcher

diff --git a/Engine/Source/Developer/DesktopPlatform/Private/Mac/DesktopPlatformMac.cpp b/Engine/Source/Developer/DesktopPlatform/Private/Mac/DesktopPlatformMac.cpp
--- a/Engine/Source/Developer/DesktopPlatform/Private/Mac/DesktopPlatformMac.cpp
+++ b/Engine/Source/Developer/DesktopPlatform/Private/Mac/DesktopPlatformMac.cpp
@@ -717,6 +717,14 @@ bool FDesktopPlatformMac::GetLauncherPath(FString& OutLauncherPath) const
 		return true;
 	}
 
+	// Otherwise search for the old Launcher...
+	Path = [Workspace fullPathForApplication:@"Unreal Engine"];
+	if( Path )
+	{
+		OutLauncherPath = FString(Path);
+		return true;
+	}
+
 	return false;
 }
 
@@ -729,6 +737,12 @@ bool FDesktopPlatformMac::GetLauncherInstallerPath(FString& OutInstallerPath) co
 		OutInstallerPath = InstallerPath;
 		return true;
 	}
+	InstallerPath = FPaths::ConvertRelativePathToFull(FPaths::Combine(*FPaths::EngineDir(), TEXT("Extras/UnrealEngineLauncher/UnrealEngine.dmg")));
+	if (FPaths::FileExists(InstallerPath))
+	{
+		OutInstallerPath = InstallerPath;
+		return true;
+	}
 	return false;
 }
 
