Allow parsing a relative path as a project's engine identifier.

diff --git a/Engine/Source/Developer/DesktopPlatform/Private/DesktopPlatformBase.cpp b/Engine/Source/Developer/DesktopPlatform/Private/DesktopPlatformBase.cpp
--- a/Engine/Source/Developer/DesktopPlatform/Private/DesktopPlatformBase.cpp
+++ b/Engine/Source/Developer/DesktopPlatform/Private/DesktopPlatformBase.cpp
@@ -129,6 +129,7 @@ bool FDesktopPlatformBase::GetEngineIdentifierFromRootDir(const FString &RootDir
 
 	// Normalize the root directory
 	FString NormalizedRootDir = RootDir;
+	FPaths::CollapseRelativeDirectories(NormalizedRootDir);
 	FPaths::NormalizeDirectoryName(NormalizedRootDir);
 
 	// Find the label for the given directory
@@ -264,7 +265,7 @@ bool FDesktopPlatformBase::SetEngineIdentifierForProject(const FString &ProjectF
 	return SaveProjectFile(ProjectFileName, ProjectFile);
 }
 
-bool FDesktopPlatformBase::GetEngineIdentifierForProject(const FString &ProjectFileName, FString &OutIdentifier)
+bool FDesktopPlatformBase::GetEngineIdentifierForProject(const FString& ProjectFileName, FString& OutIdentifier)
 {
 	OutIdentifier.Empty();
 
@@ -282,6 +283,15 @@ bool FDesktopPlatformBase::GetEngineIdentifierForProject(const FString &ProjectF
 		OutIdentifier = Value->AsString();
 		if(OutIdentifier.Len() > 0)
 		{
+			// If it's a path, convert it into an engine identifier
+			if(OutIdentifier.Contains(TEXT("/")) || OutIdentifier.Contains("\\"))
+			{
+				FString EngineRootDir = FPaths::ConvertRelativePathToFull(FPaths::GetPath(ProjectFileName), OutIdentifier);
+				if(!GetEngineIdentifierFromRootDir(EngineRootDir, OutIdentifier))
+				{
+					return false;
+				}
+			}
 			return true;
 		}
 	}
