UEPLAT-131  Enable Rocket builds for HTML5.

(Currently just on Win64 host)

diff --git a/Engine/Source/Developer/DesktopPlatform/Private/PlatformInfo.cpp b/Engine/Source/Developer/DesktopPlatform/Private/PlatformInfo.cpp
--- a/Engine/Source/Developer/DesktopPlatform/Private/PlatformInfo.cpp
+++ b/Engine/Source/Developer/DesktopPlatform/Private/PlatformInfo.cpp
@@ -105,7 +105,7 @@ static const FPlatformInfo PlatformInfoArray[] = {
 	BuildPlatformInfo(TEXT("Android_ES31"),				TEXT("Android_ES31"),		LOCTEXT("Android_ES31", "Android (ES31+AEP)"),				EPlatformType::Game,		EPlatformFlags::CookFlavor,		FPlatformIconPaths(TEXT("Launcher/AndroidTarget/Platform_Android_ES31_24x"), TEXT("Launcher/AndroidTarget/Platform_Android_128x")),					TEXT("-targetplatform=Android -cookflavor=ES31"),		TEXT("Android"),	EPlatformSDKStatus::Unknown,	TEXT("Shared/Tutorials/SettingUpAndroidTutorial"),														IsAvailableOnWindows,						false),
 //	BuildPlatformInfo(TEXT("Android_GL4"),				TEXT("Android_GL4"),		LOCTEXT("Android_GL4", "Android (GL4)"),					EPlatformType::Game,		EPlatformFlags::CookFlavor,		FPlatformIconPaths(TEXT("Launcher/AndroidTarget/Platform_Android_GL4_24x"), TEXT("Launcher/AndroidTarget/Platform_Android_128x")),					TEXT("-targetplatform=Android -cookflavor=GL4"),		TEXT("Android"),	EPlatformSDKStatus::Unknown,	TEXT("Shared/Tutorials/SettingUpAndroidTutorial"),														IsAvailableOnWindows,						true),
 
-	BuildPlatformInfo(TEXT("HTML5"),					TEXT("HTML5"),				LOCTEXT("HTML5", "HTML5"),									EPlatformType::Game,		EPlatformFlags::None,			FPlatformIconPaths(TEXT("Launcher/HTML5Target/Platform_HTML5_24x"), TEXT("Launcher/HTML5Target/Platform_HTML5_128x")),								TEXT(""),												TEXT("HTML5"),		EPlatformSDKStatus::Unknown,	TEXT("Platforms/HTML5/GettingStarted"),																	IsAvailableOnWindows || IsAvailableOnMac,   false),
+	BuildPlatformInfo(TEXT("HTML5"),					TEXT("HTML5"),				LOCTEXT("HTML5", "HTML5"),									EPlatformType::Game,		EPlatformFlags::None,			FPlatformIconPaths(TEXT("Launcher/HTML5Target/Platform_HTML5_24x"), TEXT("Launcher/HTML5Target/Platform_HTML5_128x")),								TEXT(""),												TEXT("HTML5"),		EPlatformSDKStatus::Unknown,	TEXT("Platforms/HTML5/GettingStarted"),																	IsAvailableOnWindows || IsAvailableOnMac,   true),
 
 	BuildPlatformInfo(TEXT("PS4"),						TEXT("PS4"),				LOCTEXT("PS4", "PlayStation 4"),							EPlatformType::Game,		EPlatformFlags::None,			FPlatformIconPaths(TEXT("Launcher/PS4Target/Platform_PS4_24x"), TEXT("Launcher/PS4Target/Platform_PS4_128x")),										TEXT(""),												TEXT("PS4"),		EPlatformSDKStatus::Unknown,	TEXT("Platforms/PS4/GettingStarted"),																	IsAvailableOnWindows,						false),
 
diff --git a/Engine/Source/Developer/TargetPlatform/TargetPlatform.Build.cs b/Engine/Source/Developer/TargetPlatform/TargetPlatform.Build.cs
--- a/Engine/Source/Developer/TargetPlatform/TargetPlatform.Build.cs
+++ b/Engine/Source/Developer/TargetPlatform/TargetPlatform.Build.cs
@@ -61,11 +61,7 @@ public class TargetPlatform : ModuleRules
 					DynamicallyLoadedModuleNames.Add("Android_ETC1TargetPlatform");
 					DynamicallyLoadedModuleNames.Add("Android_ETC2TargetPlatform");
 					DynamicallyLoadedModuleNames.Add("IOSTargetPlatform");
-
-					if (!UnrealBuildTool.UnrealBuildTool.RunningRocket() && !UnrealBuildTool.UnrealBuildTool.BuildingRocket())
-					{
-						DynamicallyLoadedModuleNames.Add("HTML5TargetPlatform");
-					}
+                    DynamicallyLoadedModuleNames.Add("HTML5TargetPlatform");
 				}
 			}
 			else if (Target.Platform == UnrealTargetPlatform.Mac)
diff --git a/Engine/Source/Programs/AutomationTool/Scripts/GUBP.Automation.cs b/Engine/Source/Programs/AutomationTool/Scripts/GUBP.Automation.cs
--- a/Engine/Source/Programs/AutomationTool/Scripts/GUBP.Automation.cs
+++ b/Engine/Source/Programs/AutomationTool/Scripts/GUBP.Automation.cs
@@ -5178,7 +5178,7 @@ public class GUBP : BuildCommand
             var FilteredActivePlatforms = new List<UnrealTargetPlatform>();
             foreach (var Plat in ActivePlatforms)
             {
-                if (Plat != UnrealTargetPlatform.HTML5 && 
+                if (
                     (Plat != UnrealTargetPlatform.PS4 || ParseParam("WithPS4")) &&
                     Plat != UnrealTargetPlatform.WinRT &&
                     Plat != UnrealTargetPlatform.WinRT_ARM &&
diff --git a/Engine/Source/Programs/UnrealBuildTool/HTML5/UEBuildHTML5.cs b/Engine/Source/Programs/UnrealBuildTool/HTML5/UEBuildHTML5.cs
--- a/Engine/Source/Programs/UnrealBuildTool/HTML5/UEBuildHTML5.cs
+++ b/Engine/Source/Programs/UnrealBuildTool/HTML5/UEBuildHTML5.cs
@@ -117,12 +117,6 @@ namespace UnrealBuildTool
          */
         protected override void RegisterBuildPlatformInternal()
         {
-            [REDACTED_EMAIL]: Add platform support
-            if (UnrealBuildTool.RunningRocket() || UnrealBuildTool.BuildingRocket())
-            {
-                return;
-            }
-
             // Make sure the SDK is installed
             if ((ProjectFileGenerator.bGenerateProjectFiles == true) || (HasRequiredSDKsInstalled() == SDKStatus.Valid))
             {
