Fixes for UnrealVersionSelector:

* Only try to build UBT if it's not present in the *target* engine directory (as opposed to the running engine directory). Fixes problems with running from launcher installation.
* Fix building UBT from paths that contain spaces.
* Better logging of what's going on when trying to build UBT (and possible failure points)
* Remove duplicate entries in installation list, and make sure that every path is stored in the same way. Was allowing a number of paths to contain relative portions.

diff --git a/Engine/Source/Developer/DesktopPlatform/Private/DesktopPlatformBase.cpp b/Engine/Source/Developer/DesktopPlatform/Private/DesktopPlatformBase.cpp
--- a/Engine/Source/Developer/DesktopPlatform/Private/DesktopPlatformBase.cpp
+++ b/Engine/Source/Developer/DesktopPlatform/Private/DesktopPlatformBase.cpp
@@ -465,13 +465,14 @@ bool FDesktopPlatformBase::GenerateProjectFiles(const FString& RootDir, const FS
 	// Compile UnrealBuildTool if it doesn't exist. This can happen if we're just copying source from somewhere.
 	bool bRes = true;
 	Warn->BeginSlowTask(LOCTEXT("GeneratingProjectFiles", "Generating project files..."), true, true);
-	if(!FPaths::FileExists(GetUnrealBuildToolExecutableFilename()))
+	if(!FPaths::FileExists(GetUnrealBuildToolExecutableFilename(RootDir)))
 	{
 		Warn->StatusUpdate(0, 1, LOCTEXT("BuildingUBT", "Building UnrealBuildTool..."));
-		bRes = BuildUnrealBuildTool(*Warn);
+		bRes = BuildUnrealBuildTool(RootDir, *Warn);
 	}
 	if(bRes)
 	{
+		Warn->StatusUpdate(0, 1, LOCTEXT("GeneratingProjectFiles", "Generating project files..."));
 		bRes = RunUnrealBuildTool(LOCTEXT("GeneratingProjectFiles", "Generating project files..."), RootDir, Arguments, Warn);
 	}
 	Warn->EndSlowTask();
@@ -480,24 +481,14 @@ bool FDesktopPlatformBase::GenerateProjectFiles(const FString& RootDir, const FS
 
 bool FDesktopPlatformBase::IsUnrealBuildToolAvailable()
 {
-	// If using Rocket and the Rocket unreal build tool executable exists, then UBT is available
+	// If using Rocket and the Rocket unreal build tool executable exists, then UBT is available. Otherwise check it can be built.
 	if (FApp::IsEngineInstalled())
 	{
-		return FPaths::FileExists(GetUnrealBuildToolExecutableFilename());
+		return FPaths::FileExists(GetUnrealBuildToolExecutableFilename(FPaths::RootDir()));
 	}
 	else
 	{
-		// If not using Rocket, check to make sure UBT can be built, since it is an intermediate.
-		// We are simply checking for the existence of source code files, which is a heuristic to determine if UBT can be built
-		TArray<FString> Filenames;
-		const FString UBTSourcePath = GetUnrealBuildToolSourceCodePath();
-		const FString SearchPattern = TEXT("*.cs");
-		const bool bFiles = true;
-		const bool bDirectories = false;
-		const bool bClearFileNames = false;
-		IFileManager::Get().FindFilesRecursive(Filenames, *UBTSourcePath, *SearchPattern, bFiles, bDirectories, bClearFileNames);
-
-		return Filenames.Num() > 0;
+		return FPaths::FileExists(GetUnrealBuildToolProjectFileName(FPaths::RootDir()));
 	}
 }
 
@@ -553,7 +544,7 @@ FProcHandle FDesktopPlatformBase::InvokeUnrealBuildToolAsync(const FString& InCm
 	}
 
 	// UnrealBuildTool is currently always located in the Binaries/DotNET folder
-	FString ExecutableFileName = GetUnrealBuildToolExecutableFilename();
+	FString ExecutableFileName = GetUnrealBuildToolExecutableFilename(FPaths::RootDir());
 
 	// Rocket never builds UBT, UnrealBuildTool should already exist
 	bool bSkipBuild = FApp::IsEngineInstalled() || bSkipBuildUBT;
@@ -565,7 +556,7 @@ FProcHandle FDesktopPlatformBase::InvokeUnrealBuildToolAsync(const FString& InCm
 		if (!bSuccessfullyBuiltUBTOnce)
 		{
 			Ar.Log(TEXT("Building UnrealBuildTool..."));
-			if (BuildUnrealBuildTool(Ar))
+			if (BuildUnrealBuildTool(FPaths::RootDir(), Ar))
 			{
 				bSuccessfullyBuiltUBTOnce = true;
 			}
@@ -914,23 +905,18 @@ bool FDesktopPlatformBase::EnumerateProjectsKnownByEngine(const FString &Identif
 	return true;
 }
 
-FString FDesktopPlatformBase::GetUnrealBuildToolExecutableFilename() const
+bool FDesktopPlatformBase::BuildUnrealBuildTool(const FString& RootDir, FOutputDevice& Ar)
 {
-	const FString UBTPath = FString::Printf(TEXT("%sBinaries/DotNET/"), *FPaths::EngineDir());
-	const FString UBTExe = TEXT("UnrealBuildTool.exe");
-
-	FString ExecutableFileName = UBTPath / UBTExe;
-	ExecutableFileName = FPaths::ConvertRelativePathToFull(ExecutableFileName);
-	return ExecutableFileName;
-}
+	Ar.Logf(TEXT("Building UnrealBuildTool in %s..."), *RootDir);
 
-FString FDesktopPlatformBase::GetUnrealBuildToolSourceCodePath() const
-{
-	return FPaths::Combine(*FPaths::EngineDir(), TEXT("Source"), TEXT("Programs"), TEXT("UnrealBuildTool"));
-}
+	// Check the project file exists
+	FString CsProjLocation = GetUnrealBuildToolProjectFileName(RootDir);
+	if(!FPaths::FileExists(CsProjLocation))
+	{
+		Ar.Logf(TEXT("Project file not found at %s"), *CsProjLocation);
+		return false;
+	}
 
-bool FDesktopPlatformBase::BuildUnrealBuildTool(FOutputDevice &Ar)
-{
 	FString CompilerExecutableFilename;
 	FString CmdLineParams;
 #if PLATFORM_WINDOWS
@@ -952,63 +938,70 @@ bool FDesktopPlatformBase::BuildUnrealBuildTool(FOutputDevice &Ar)
 	// Check to make sure we found one.
 	if (VCVarsBat.IsEmpty() || !FPaths::FileExists(VCVarsBat))
 	{
-		// VCVars doesn't exist, we can not build UBT
-		CompilerExecutableFilename = TEXT("");
+		Ar.Logf(TEXT("Couldn't find %s; skipping."), *VCVarsBat);
+		return false;
 	}
-	else
-	{
-		// Now make a batch file in the intermediate directory to invoke the vcvars batch then msbuild
-		FString BuildBatchFile = FPaths::EngineIntermediateDir() / TEXT("Build") / TEXT("UnrealBuildTool") / TEXT("BuildUBT.bat");
-		BuildBatchFile.ReplaceInline(TEXT("/"), TEXT("\\"));
 
-		const FString CsProjLocation = FPaths::ConvertRelativePathToFull(GetUnrealBuildToolSourceCodePath()) / TEXT("UnrealBuildTool.csproj");
+	// Now make a batch file in the intermediate directory to invoke the vcvars batch then msbuild
+	FString BuildBatchFile = RootDir / TEXT("Engine/Intermediate/Build/UnrealBuildTool/BuildUBT.bat");
+	BuildBatchFile.ReplaceInline(TEXT("/"), TEXT("\\"));
 
-		FString BatchFileContents;
-		BatchFileContents = FString::Printf(TEXT("call \"%s\"") LINE_TERMINATOR, *VCVarsBat);
-		BatchFileContents += FString::Printf(TEXT("msbuild /nologo /verbosity:quiet \"%s\" /property:Configuration=Development /property:Platform=AnyCPU"), *CsProjLocation);
-		FFileHelper::SaveStringToFile(BatchFileContents, *BuildBatchFile);
+	FString BatchFileContents;
+	BatchFileContents = FString::Printf(TEXT("call \"%s\"") LINE_TERMINATOR, *VCVarsBat);
+	BatchFileContents += FString::Printf(TEXT("msbuild /nologo /verbosity:quiet \"%s\" /property:Configuration=Development /property:Platform=AnyCPU"), *CsProjLocation);
+	FFileHelper::SaveStringToFile(BatchFileContents, *BuildBatchFile);
 
-		{
-			TCHAR CmdExePath[MAX_PATH];
-			FPlatformMisc::GetEnvironmentVariable(TEXT("ComSpec"), CmdExePath, ARRAY_COUNT(CmdExePath));
-			CompilerExecutableFilename = CmdExePath;
-		}
+	TCHAR CmdExePath[MAX_PATH];
+	FPlatformMisc::GetEnvironmentVariable(TEXT("ComSpec"), CmdExePath, ARRAY_COUNT(CmdExePath));
+	CompilerExecutableFilename = CmdExePath;
 
-		CmdLineParams = TEXT("/c ");
-		CmdLineParams += BuildBatchFile;
-	}
+	CmdLineParams = FString::Printf(TEXT("/c \"%s\""), *BuildBatchFile);
 #elif PLATFORM_MAC
-	const FString CsProjLocation = FPaths::ConvertRelativePathToFull(GetUnrealBuildToolSourceCodePath()) / TEXT("UnrealBuildTool.csproj");
-	FString ScriptPath = FPaths::ConvertRelativePathToFull(FPaths::EngineDir() / TEXT("Build/BatchFiles/Mac/RunXBuild.sh"));
+	FString ScriptPath = FPaths::ConvertRelativePathToFull(RootDir / TEXT("Engine/Build/BatchFiles/Mac/RunXBuild.sh"));
 	CompilerExecutableFilename = TEXT("/bin/sh");
 	CmdLineParams = FString::Printf(TEXT("\"%s\" /property:Configuration=Development %s"), *ScriptPath, *CsProjLocation);
 #elif PLATFORM_LINUX
-	const FString CsProjLocation = FPaths::ConvertRelativePathToFull(GetUnrealBuildToolSourceCodePath()) / TEXT("UnrealBuildTool.csproj");
-	FString ScriptPath = FPaths::ConvertRelativePathToFull(FPaths::EngineDir() / TEXT("Build/BatchFiles/Linux/RunXBuild.sh"));
+	FString ScriptPath = FPaths::ConvertRelativePathToFull(RootDir / TEXT("Build/BatchFiles/Linux/RunXBuild.sh"));
 	CompilerExecutableFilename = TEXT("/bin/bash");
 	CmdLineParams = FString::Printf(TEXT("\"%s\" /property:Configuration=Development /property:TargetFrameworkVersion=v4.0 %s"), *ScriptPath, *CsProjLocation);
 #else
 	Ar.Log(TEXT("Unknown platform, unable to build UnrealBuildTool."));
+	return false;
 #endif
 
-	// If a compiler executable was provided, try to build now
-	if (CompilerExecutableFilename.IsEmpty())
-	{
-		return false;
-	}
-
+	// Spawn the compiler
+	Ar.Logf(TEXT("Running: %s %s"), *CompilerExecutableFilename, *CmdLineParams);
 	const bool bLaunchDetached = false;
 	const bool bLaunchHidden = true;
 	const bool bLaunchReallyHidden = bLaunchHidden;
 	FProcHandle ProcHandle = FPlatformProcess::CreateProc(*CompilerExecutableFilename, *CmdLineParams, bLaunchDetached, bLaunchHidden, bLaunchReallyHidden, NULL, 0, NULL, NULL);
-	if (ProcHandle.IsValid())
+	if (!ProcHandle.IsValid())
 	{
-		FPlatformProcess::WaitForProc(ProcHandle);
-		ProcHandle.Close();
+		Ar.Log(TEXT("Failed to start process."));
+		return false;
 	}
+	FPlatformProcess::WaitForProc(ProcHandle);
+	ProcHandle.Close();
 
 	// If the executable appeared where we expect it, then we were successful
-	return FPaths::FileExists(GetUnrealBuildToolExecutableFilename());
+	FString UnrealBuildToolExePath = GetUnrealBuildToolExecutableFilename(RootDir);
+	if(!FPaths::FileExists(UnrealBuildToolExePath))
+	{
+		Ar.Logf(TEXT("Missing %s after build"), *UnrealBuildToolExePath);
+		return false;
+	}
+
+	return true;
+}
+
+FString FDesktopPlatformBase::GetUnrealBuildToolProjectFileName(const FString& RootDir) const
+{
+	return FPaths::ConvertRelativePathToFull(RootDir / TEXT("Engine/Source/Programs/UnrealBuildTool/UnrealBuildTool.csproj"));
+}
+
+FString FDesktopPlatformBase::GetUnrealBuildToolExecutableFilename(const FString& RootDir) const
+{
+	return FPaths::ConvertRelativePathToFull(RootDir / TEXT("Engine/Binaries/DotNET/UnrealBuildTool.exe"));
 }
 
 #undef LOCTEXT_NAMESPACE
diff --git a/Engine/Source/Developer/DesktopPlatform/Private/DesktopPlatformBase.h b/Engine/Source/Developer/DesktopPlatform/Private/DesktopPlatformBase.h
--- a/Engine/Source/Developer/DesktopPlatform/Private/DesktopPlatformBase.h
+++ b/Engine/Source/Developer/DesktopPlatform/Private/DesktopPlatformBase.h
@@ -63,11 +63,10 @@ private:
 	const FUProjectDictionary &GetCachedProjectDictionary(const FString& RootDir);
 
 	void GetProjectBuildProducts(const FString& ProjectFileName, TArray<FString> &OutFileNames, TArray<FString> &OutDirectoryNames);
-
-	FString GetUnrealBuildToolSourceCodePath() const;
-	bool BuildUnrealBuildTool(FOutputDevice &Ar);
+	bool BuildUnrealBuildTool(const FString& RootDir, FOutputDevice &Ar);
 
 protected:
 
-	FString GetUnrealBuildToolExecutableFilename() const;	
+	FString GetUnrealBuildToolProjectFileName(const FString& RootDir) const;
+	FString GetUnrealBuildToolExecutableFilename(const FString& RootDir) const;	
 };
diff --git a/Engine/Source/Developer/DesktopPlatform/Private/Windows/DesktopPlatformWindows.cpp b/Engine/Source/Developer/DesktopPlatform/Private/Windows/DesktopPlatformWindows.cpp
--- a/Engine/Source/Developer/DesktopPlatform/Private/Windows/DesktopPlatformWindows.cpp
+++ b/Engine/Source/Developer/DesktopPlatform/Private/Windows/DesktopPlatformWindows.cpp
@@ -424,7 +424,12 @@ void FDesktopPlatformWindows::EnumerateEngineInstallations(TMap<FString, FString
 	HKEY hKey;
 	if (RegOpenKeyEx(HKEY_CURRENT_USER, InstallationsSubKey, 0, KEY_ALL_ACCESS, &hKey) == ERROR_SUCCESS)
 	{
+		// Get a list of all the directories
+		TArray<FString> UniqueDirectories;
+		OutInstallations.GenerateValueArray(UniqueDirectories);
+
 		// Enumerate all the installations
+		TArray<FString> InvalidKeys;
 		for (::DWORD Index = 0;; Index++)
 		{
 			TCHAR ValueName[256];
@@ -441,7 +446,17 @@ void FDesktopPlatformWindows::EnumerateEngineInstallations(TMap<FString, FString
 
 				FString NormalizedInstalledDirectory(ValueDataLength, ValueData);
 				FPaths::NormalizeDirectoryName(NormalizedInstalledDirectory);
-				OutInstallations.Add(ValueName, NormalizedInstalledDirectory);
+				FPaths::CollapseRelativeDirectories(NormalizedInstalledDirectory);
+
+				if(IsValidRootDirectory(NormalizedInstalledDirectory) && !UniqueDirectories.Contains(NormalizedInstalledDirectory))
+				{
+					OutInstallations.Add(ValueName, NormalizedInstalledDirectory);
+					UniqueDirectories.Add(NormalizedInstalledDirectory);
+				}
+				else
+				{
+					InvalidKeys.Add(ValueName);
+				}
 			}
 			else if(Result == ERROR_NO_MORE_ITEMS)
 			{
@@ -449,14 +464,10 @@ void FDesktopPlatformWindows::EnumerateEngineInstallations(TMap<FString, FString
 			}
 		}
 
-		// Trim the list to everything that's valid
-		for(TMap<FString, FString>::TIterator Iter(OutInstallations); Iter; ++Iter)
+		// Remove all the keys which weren't valid
+		for(const FString InvalidKey: InvalidKeys)
 		{
-			if (!IsValidRootDirectory(Iter.Value()))
-			{
-				RegDeleteValue(hKey, *Iter.Key());
-				Iter.RemoveCurrent();
-			}
+			RegDeleteValue(hKey, *InvalidKey);
 		}
 
 		RegCloseKey(hKey);
@@ -546,7 +557,7 @@ bool FDesktopPlatformWindows::RunUnrealBuildTool(const FText& Description, const
 
 bool FDesktopPlatformWindows::IsUnrealBuildToolRunning()
 {
-	FString UBTPath = GetUnrealBuildToolExecutableFilename();
+	FString UBTPath = GetUnrealBuildToolExecutableFilename(FPaths::RootDir());
 	FPaths::MakePlatformFilename(UBTPath);
 
 	HANDLE SnapShot = ::CreateToolhelp32Snapshot(TH32CS_SNAPPROCESS, 0);
