[INTEGRATE] Change 2431733 by [REDACTED_PERSON][REDACTED_HANDLE] on 2015/02/04 10:34:16

	Fix non-Rocket projects generated even with "-Rocket" passed

diff --git a/Engine/Source/Developer/DesktopPlatform/Private/DesktopPlatformBase.cpp b/Engine/Source/Developer/DesktopPlatform/Private/DesktopPlatformBase.cpp
--- a/Engine/Source/Developer/DesktopPlatform/Private/DesktopPlatformBase.cpp
+++ b/Engine/Source/Developer/DesktopPlatform/Private/DesktopPlatformBase.cpp
@@ -450,7 +450,7 @@ bool FDesktopPlatformBase::GenerateProjectFiles(const FString& RootDir, const FS
 			Arguments += " -game";
 
 			// Determine whether or not to include engine source
-			if(IsSourceDistribution(RootDir))
+			if(IsSourceDistribution(RootDir) && !FRocketSupport::IsRocket())
 			{
 				Arguments += " -engine";
 			}
