fix for failure to cook from File|Cook
UE-9786

diff --git a/Engine/Source/Developer/DesktopPlatform/Private/PlatformInfo.cpp b/Engine/Source/Developer/DesktopPlatform/Private/PlatformInfo.cpp
--- a/Engine/Source/Developer/DesktopPlatform/Private/PlatformInfo.cpp
+++ b/Engine/Source/Developer/DesktopPlatform/Private/PlatformInfo.cpp
@@ -76,19 +76,19 @@ static const FString IOSTutorial = TEXT("/Engine/Tutorial/NotYetImplemented");
 
 static const FPlatformInfo PlatformInfoArray[] = {
 	//				  PlatformInfoName					TargetPlatformName			DisplayName													PlatformType				PlatformFlags					IconPaths																																			UATCommandLine											AutoSDKPath			SDKStatus						SDKTutorial																								EnabledForUse								EnabledForBinary	UsesHostCompiler
-	BuildPlatformInfo(TEXT("WindowsNoEditor"),			TEXT("WindowsNoEditor"),	LOCTEXT("WindowsNoEditor", "Windows"),						EPlatformType::Game,		EPlatformFlags::None,			FPlatformIconPaths(TEXT("Launcher/WindowsTarget/Platform_WindowsNoEditor_24x"), TEXT("Launcher/WindowsTarget/Platform_WindowsNoEditor_128x")),		TEXT(""),												TEXT(""),			EPlatformSDKStatus::Unknown,	TEXT("/Engine/Tutorial/Installation/InstallingVisualStudioTutorial.InstallingVisualStudioTutorial"),	IsAvailableOnWindows,						true,				IsAvailableOnWindows),
+	BuildPlatformInfo(TEXT("WindowsNoEditor"),			TEXT("WindowsNoEditor"),	LOCTEXT("WindowsNoEditor", "Windows"),						EPlatformType::Game,		EPlatformFlags::None,			FPlatformIconPaths(TEXT("Launcher/WindowsTarget/Platform_WindowsNoEditor_24x"), TEXT("Launcher/WindowsTarget/Platform_WindowsNoEditor_128x")),		TEXT("-targetplatform=Win64"),							TEXT(""),			EPlatformSDKStatus::Unknown,	TEXT("/Engine/Tutorial/Installation/InstallingVisualStudioTutorial.InstallingVisualStudioTutorial"),	IsAvailableOnWindows,						true,				IsAvailableOnWindows),
 	BuildPlatformInfo(TEXT("WindowsNoEditor_Win32"),	TEXT("WindowsNoEditor"),	LOCTEXT("WindowsNoEditor_Win32", "Windows (32-bit)"),		EPlatformType::Game,		EPlatformFlags::BuildFlavor,	FPlatformIconPaths(TEXT("Launcher/WindowsTarget/Platform_WindowsNoEditor_24x"), TEXT("Launcher/WindowsTarget/Platform_WindowsNoEditor_128x")),		TEXT("-targetplatform=Win32"),							TEXT(""),			EPlatformSDKStatus::Unknown,	TEXT("/Engine/Tutorial/Installation/InstallingVisualStudioTutorial.InstallingVisualStudioTutorial"),	IsAvailableOnWindows,						true,				IsAvailableOnWindows),
 	BuildPlatformInfo(TEXT("WindowsNoEditor_Win64"),	TEXT("WindowsNoEditor"),	LOCTEXT("WindowsNoEditor_Win64", "Windows (64-bit)"),		EPlatformType::Game,		EPlatformFlags::BuildFlavor,	FPlatformIconPaths(TEXT("Launcher/WindowsTarget/Platform_WindowsNoEditor_24x"), TEXT("Launcher/WindowsTarget/Platform_WindowsNoEditor_128x")),		TEXT("-targetplatform=Win64"),							TEXT(""),			EPlatformSDKStatus::Unknown,	TEXT("/Engine/Tutorial/Installation/InstallingVisualStudioTutorial.InstallingVisualStudioTutorial"),	IsAvailableOnWindows,						true,				IsAvailableOnWindows),
 	BuildPlatformInfo(TEXT("Windows"),					TEXT("Windows"),			LOCTEXT("WindowsEditor", "Windows (Editor)"),				EPlatformType::Editor,		EPlatformFlags::None,			FPlatformIconPaths(TEXT("Launcher/WindowsTarget/Platform_Windows_24x"), TEXT("Launcher/WindowsTarget/Platform_Windows_128x")),						TEXT(""),												TEXT(""),			EPlatformSDKStatus::Unknown,	TEXT("/Engine/Tutorial/Installation/InstallingVisualStudioTutorial.InstallingVisualStudioTutorial"),	IsAvailableOnWindows,						true,				IsAvailableOnWindows),
 	BuildPlatformInfo(TEXT("WindowsClient"),			TEXT("WindowsClient"),		LOCTEXT("WindowsClient", "Windows (Client-only)"),			EPlatformType::Client,		EPlatformFlags::None,			FPlatformIconPaths(TEXT("Launcher/WindowsTarget/Platform_Windows_24x"), TEXT("Launcher/WindowsTarget/Platform_Windows_128x")),						TEXT(""),												TEXT(""),			EPlatformSDKStatus::Unknown,	TEXT("/Engine/Tutorial/Installation/InstallingVisualStudioTutorial.InstallingVisualStudioTutorial"),	IsAvailableOnWindows,						true,				IsAvailableOnWindows),
 	BuildPlatformInfo(TEXT("WindowsServer"),			TEXT("WindowsServer"),		LOCTEXT("WindowsServer", "Windows (Dedicated Server)"),		EPlatformType::Server,		EPlatformFlags::None,			FPlatformIconPaths(TEXT("Launcher/WindowsTarget/Platform_WindowsServer_24x"), TEXT("Launcher/WindowsTarget/Platform_WindowsServer_128x")),			TEXT(""),												TEXT(""),			EPlatformSDKStatus::Unknown,	TEXT("/Engine/Tutorial/Installation/InstallingVisualStudioTutorial.InstallingVisualStudioTutorial"),	IsAvailableOnWindows,						true,				IsAvailableOnWindows),
 
-	BuildPlatformInfo(TEXT("MacNoEditor"),				TEXT("MacNoEditor"),		LOCTEXT("MacNoEditor", "Mac"),								EPlatformType::Game,		EPlatformFlags::None,			FPlatformIconPaths(TEXT("Launcher/MacTarget/Platform_Mac_24x"), TEXT("Launcher/MacTarget/Platform_Mac_128x")),										TEXT(""),												TEXT(""),			EPlatformSDKStatus::Unknown,	TEXT("/Engine/Tutorial/Installation/InstallingXCodeTutorial.InstallingXCodeTutorial"),					IsAvailableOnMac,							true,				IsAvailableOnMac),
+	BuildPlatformInfo(TEXT("MacNoEditor"),				TEXT("MacNoEditor"),		LOCTEXT("MacNoEditor", "Mac"),								EPlatformType::Game,		EPlatformFlags::None,			FPlatformIconPaths(TEXT("Launcher/MacTarget/Platform_Mac_24x"), TEXT("Launcher/MacTarget/Platform_Mac_128x")),										TEXT("-targetplatform=Mac"),							TEXT(""),			EPlatformSDKStatus::Unknown,	TEXT("/Engine/Tutorial/Installation/InstallingXCodeTutorial.InstallingXCodeTutorial"),					IsAvailableOnMac,							true,				IsAvailableOnMac),
 	BuildPlatformInfo(TEXT("Mac"),						TEXT("Mac"),				LOCTEXT("MacEditor", "Mac (Editor)"),						EPlatformType::Editor,		EPlatformFlags::None,			FPlatformIconPaths(TEXT("Launcher/MacTarget/Platform_Mac_24x"), TEXT("Launcher/MacTarget/Platform_Mac_128x")),										TEXT(""),												TEXT(""),			EPlatformSDKStatus::Unknown,	TEXT("/Engine/Tutorial/Installation/InstallingXCodeTutorial.InstallingXCodeTutorial"),					IsAvailableOnMac,							true,				IsAvailableOnMac),
 	BuildPlatformInfo(TEXT("MacClient"),				TEXT("MacClient"),			LOCTEXT("MacClient", "Mac (Client-only)"),					EPlatformType::Client,		EPlatformFlags::None,			FPlatformIconPaths(TEXT("Launcher/MacTarget/Platform_Mac_24x"), TEXT("Launcher/MacTarget/Platform_Mac_128x")),										TEXT(""),												TEXT(""),			EPlatformSDKStatus::Unknown,	TEXT("/Engine/Tutorial/Installation/InstallingXCodeTutorial.InstallingXCodeTutorial"),					IsAvailableOnMac,							true,				IsAvailableOnMac),
 	BuildPlatformInfo(TEXT("MacServer"),				TEXT("MacServer"),			LOCTEXT("MacServer", "Mac (Dedicated Server)"),				EPlatformType::Server,		EPlatformFlags::None,			FPlatformIconPaths(TEXT("Launcher/MacTarget/Platform_Mac_24x"), TEXT("Launcher/MacTarget/Platform_Mac_128x")),										TEXT(""),												TEXT(""),			EPlatformSDKStatus::Unknown,	TEXT("/Engine/Tutorial/Installation/InstallingXCodeTutorial.InstallingXCodeTutorial"),					IsAvailableOnMac,							true,				IsAvailableOnMac),
 
-	BuildPlatformInfo(TEXT("LinuxNoEditor"),			TEXT("LinuxNoEditor"),		LOCTEXT("LinuxNoEditor", "Linux"),							EPlatformType::Game,		EPlatformFlags::None,			FPlatformIconPaths(TEXT("Launcher/LinuxTarget/Platform_Linux_24x"), TEXT("Launcher/LinuxTarget/Platform_Linux_128x")),								TEXT(""),												TEXT("Linux_x64"),	EPlatformSDKStatus::Unknown,	TEXT(""),																								IsAvailableOnLinux || IsAvailableOnWindows,	 true,				IsAvailableOnLinux),	
+	BuildPlatformInfo(TEXT("LinuxNoEditor"),			TEXT("LinuxNoEditor"),		LOCTEXT("LinuxNoEditor", "Linux"),							EPlatformType::Game,		EPlatformFlags::None,			FPlatformIconPaths(TEXT("Launcher/LinuxTarget/Platform_Linux_24x"), TEXT("Launcher/LinuxTarget/Platform_Linux_128x")),								TEXT("-targetplatform=Linux"),							TEXT("Linux_x64"),	EPlatformSDKStatus::Unknown,	TEXT(""),																								IsAvailableOnLinux || IsAvailableOnWindows,	 true,				IsAvailableOnLinux),	
 	BuildPlatformInfo(TEXT("Linux"),					TEXT("Linux"),				LOCTEXT("LinuxEditor", "Linux (Editor)"),					EPlatformType::Editor,		EPlatformFlags::None,			FPlatformIconPaths(TEXT("Launcher/LinuxTarget/Platform_Linux_24x"), TEXT("Launcher/LinuxTarget/Platform_Linux_128x")),								TEXT(""),												TEXT("Linux_x64"),	EPlatformSDKStatus::Unknown,	TEXT(""),																								IsAvailableOnLinux,							 true,				IsAvailableOnLinux),
 	BuildPlatformInfo(TEXT("LinuxClient"),				TEXT("LinuxClient"),		LOCTEXT("LinuxClient", "Linux (Client-only)"),				EPlatformType::Client,		EPlatformFlags::None,			FPlatformIconPaths(TEXT("Launcher/LinuxTarget/Platform_Linux_24x"), TEXT("Launcher/LinuxTarget/Platform_Linux_128x")),								TEXT(""),												TEXT("Linux_x64"),	EPlatformSDKStatus::Unknown,	TEXT(""),																								IsAvailableOnLinux || IsAvailableOnWindows,	 true,				IsAvailableOnLinux),
 	BuildPlatformInfo(TEXT("LinuxServer"),				TEXT("LinuxServer"),		LOCTEXT("LinuxServer", "Linux (Dedicated Server)"),			EPlatformType::Server,		EPlatformFlags::None,			FPlatformIconPaths(TEXT("Launcher/LinuxTarget/Platform_Linux_24x"), TEXT("Launcher/LinuxTarget/Platform_Linux_128x")),								TEXT(""),												TEXT("Linux_x64"),	EPlatformSDKStatus::Unknown,	TEXT(""),																								IsAvailableOnLinux || IsAvailableOnWindows,	 true,				IsAvailableOnLinux),
diff --git a/Engine/Source/Editor/MainFrame/Private/Frame/MainFrameActions.cpp b/Engine/Source/Editor/MainFrame/Private/Frame/MainFrameActions.cpp
--- a/Engine/Source/Editor/MainFrame/Private/Frame/MainFrameActions.cpp
+++ b/Engine/Source/Editor/MainFrame/Private/Frame/MainFrameActions.cpp
@@ -312,21 +312,17 @@ void FMainFrameActionCallbacks::CookContent(const FName InPlatformInfoName)
 		return;
 	}
 
-	if (PlatformInfo->TargetPlatformName == FName("MacNoEditor"))
-	{
-		OptionalParams += TEXT(" -targetplatform=Mac");
-	}
-	else if (PlatformInfo->TargetPlatformName == FName("LinuxNoEditor"))
-	{
-		OptionalParams += TEXT(" -targetplatform=Linux");
-	}
-
 	// Append any extra UAT flags specified for this platform flavor
 	if (!PlatformInfo->UATCommandLine.IsEmpty())
 	{
 		OptionalParams += TEXT(" ");
 		OptionalParams += PlatformInfo->UATCommandLine;
 	}
+	else
+	{
+		OptionalParams += TEXT(" -targetplatform=");
+		OptionalParams += *PlatformInfo->TargetPlatformName.ToString();
+	}
 
 	const bool bRunningDebug = FParse::Param(FCommandLine::Get(), TEXT("debug"));
 
@@ -336,13 +332,12 @@ void FMainFrameActionCallbacks::CookContent(const FName InPlatformInfoName)
 	}
 
 	FString ProjectPath = FPaths::IsProjectFilePathSet() ? FPaths::ConvertRelativePathToFull(FPaths::GetProjectFilePath()) : FPaths::RootDir() / FApp::GetGameName() / FApp::GetGameName() + TEXT(".uproject");
-	FString CommandLine = FString::Printf(TEXT("BuildCookRun %s%s%s%s -nop4 -project=\"%s\" -cook -platform=%s -ue4exe=%s %s"),
+	FString CommandLine = FString::Printf(TEXT("BuildCookRun %s%s%s%s -nop4 -project=\"%s\" -cook -ue4exe=%s %s"),
 		FRocketSupport::IsRocket() ? TEXT("-rocket ") : TEXT(""),
 		GetUATCompilationFlags(),
 		GetUATNoCodeProjectFlag(),
 		FApp::IsEngineInstalled() ? TEXT(" -installed") : TEXT(""),
 		*ProjectPath,
-		*PlatformInfo->TargetPlatformName.ToString(),
 		*FUnrealEdMisc::Get().GetExecutableForCommandlets(),
 		*OptionalParams
 	);
@@ -535,15 +530,7 @@ void FMainFrameActionCallbacks::PackageProject( const FName InPlatformInfoName )
 		OptionalParams += TEXT(" -distribution");
 	}
 
-	if (PlatformInfo->TargetPlatformName == FName("MacNoEditor"))
-	{
-		OptionalParams += TEXT(" -targetplatform=Mac");
-	}
-	else if (PlatformInfo->TargetPlatformName == FName("LinuxNoEditor"))
-	{
-		OptionalParams += TEXT(" -targetplatform=Linux");
-	}
-	else if (PlatformInfo->TargetPlatformName == FName("WindowsNoEditor") && PlatformInfo->PlatformFlavor == TEXT("Win32"))
+	if (PlatformInfo->TargetPlatformName == FName("WindowsNoEditor") && PlatformInfo->PlatformFlavor == TEXT("Win32"))
 	{
 		FString MinumumSupportedWindowsOS;
 		GConfig->GetString(TEXT("/Script/WindowsTargetPlatform.WindowsTargetSettings"), TEXT("MinimumOSVersion"), MinumumSupportedWindowsOS, GEngineIni);
@@ -559,6 +546,11 @@ void FMainFrameActionCallbacks::PackageProject( const FName InPlatformInfoName )
 		OptionalParams += TEXT(" ");
 		OptionalParams += PlatformInfo->UATCommandLine;
 	}
+	else
+	{
+		OptionalParams += TEXT(" -targetplatform=");
+		OptionalParams += *PlatformInfo->TargetPlatformName.ToString();
+	}
 
 	// only build if the project has code that might need to be built
 	if (bProjectHasCode && FSourceCodeNavigation::IsCompilerAvailable())
@@ -596,14 +588,13 @@ void FMainFrameActionCallbacks::PackageProject( const FName InPlatformInfoName )
 	Configuration = Configuration.Replace(TEXT("PPBC_"), TEXT(""));
 
 	FString ProjectPath = FPaths::IsProjectFilePathSet() ? FPaths::ConvertRelativePathToFull(FPaths::GetProjectFilePath()) : FPaths::RootDir() / FApp::GetGameName() / FApp::GetGameName() + TEXT(".uproject");
-	FString CommandLine = FString::Printf(TEXT("BuildCookRun %s%s%s%s -nop4 -project=\"%s\" -cook -stage -archive -archivedirectory=\"%s\" -package -platform=%s -clientconfig=%s -ue4exe=%s %s -utf8output"),
+	FString CommandLine = FString::Printf(TEXT("BuildCookRun %s%s%s%s -nop4 -project=\"%s\" -cook -stage -archive -archivedirectory=\"%s\" -package -clientconfig=%s -ue4exe=%s %s -utf8output"),
 		FRocketSupport::IsRocket() ? TEXT("-rocket ") : TEXT(""),
 		GetUATCompilationFlags(),
 		GetUATNoCodeProjectFlag(),
 		FApp::IsEngineInstalled() ? TEXT(" -installed") : TEXT(""),
 		*ProjectPath,
 		*PackagingSettings->StagingDirectory.Path,
-		*PlatformInfo->TargetPlatformName.ToString(),
 		*Configuration,
 		*ExecutableName,
 		*OptionalParams
