Added support for "gathering" project files via the editor

This is needed for UBT to update its makefiles when we add new code via the editor.

Updating UBT makefiles is more typically handled by re-generating project files, however on Windows we use DTE to inject new files directly into Visual Studio in an attempt to avoid re-generating project files (as it produces modal popups that block operations), so we need to perform a gather instead so that UBT knows things have changed.

diff --git a/Engine/Source/Developer/DesktopPlatform/Private/DesktopPlatformBase.cpp b/Engine/Source/Developer/DesktopPlatform/Private/DesktopPlatformBase.cpp
--- a/Engine/Source/Developer/DesktopPlatform/Private/DesktopPlatformBase.cpp
+++ b/Engine/Source/Developer/DesktopPlatform/Private/DesktopPlatformBase.cpp
@@ -430,13 +430,13 @@ bool FDesktopPlatformBase::CompileGameProject(const FString& RootDir, const FStr
 bool FDesktopPlatformBase::GenerateProjectFiles(const FString& RootDir, const FString& ProjectFileName, FFeedbackContext* Warn)
 {
 #if PLATFORM_MAC
-	FString Arguments = TEXT("-xcodeprojectfile");
+	FString Arguments = TEXT(" -xcodeprojectfile");
 #elif PLATFORM_LINUX
-	FString Arguments = TEXT(" -makefile -qmakefile -cmakefile ");
+	FString Arguments = TEXT(" -makefile -qmakefile -cmakefile");
 #else
-	FString Arguments = TEXT("-projectfiles");
+	FString Arguments = TEXT(" -projectfiles");
 #endif
-
+	
 	// Build the arguments to pass to UBT. If it's a non-foreign project, just build full project files.
 	if ( !ProjectFileName.IsEmpty() && GetCachedProjectDictionary(RootDir).IsForeignProject(ProjectFileName) )
 	{
@@ -447,20 +447,20 @@ bool FDesktopPlatformBase::GenerateProjectFiles(const FString& RootDir, const FS
 			Arguments += FString::Printf(TEXT(" -project=\"%s\""), *IFileManager::Get().ConvertToAbsolutePathForExternalAppForRead(*ProjectFileName));
 
 			// Always include game source
-			Arguments += " -game";
+			Arguments += TEXT(" -game");
 
 			// Determine whether or not to include engine source
 			if(IsSourceDistribution(RootDir) && !FRocketSupport::IsRocket())
 			{
-				Arguments += " -engine";
+				Arguments += TEXT(" -engine");
 			}
 			else
 			{
-				Arguments += " -rocket";
+				Arguments += TEXT(" -rocket");
 			}
 		}
 	}
-	Arguments += " -progress";
+	Arguments += TEXT(" -progress");
 
 	// Compile UnrealBuildTool if it doesn't exist. This can happen if we're just copying source from somewhere.
 	bool bRes = true;
@@ -479,6 +479,52 @@ bool FDesktopPlatformBase::GenerateProjectFiles(const FString& RootDir, const FS
 	return bRes;
 }
 
+bool FDesktopPlatformBase::GatherProjectFiles(const FString& RootDir, const FString& ProjectFileName, FFeedbackContext* Warn)
+{
+	// Composes the target, platform, and config (eg, "QAGame Win64 Development")
+	FString Arguments = FString::Printf(TEXT("%s %s %s"), FApp::GetGameName(), FPlatformMisc::GetUBTPlatform(), FModuleManager::GetUBTConfiguration());
+
+	// -editorrecompile tells UBT to work out the editor target name from the game target name we provided (eg, converting "QAGame" to "QAGameEditor")
+	Arguments += TEXT(" -editorrecompile");
+
+	// Build the arguments to pass to UBT. If it's a non-foreign project, just build full project files.
+	if ( !ProjectFileName.IsEmpty() && GetCachedProjectDictionary(RootDir).IsForeignProject(ProjectFileName) )
+	{
+		// Figure out whether it's a foreign project
+		const FUProjectDictionary &ProjectDictionary = GetCachedProjectDictionary(RootDir);
+		if(ProjectDictionary.IsForeignProject(ProjectFileName))
+		{
+			Arguments += FString::Printf(TEXT(" \"%s\""), *IFileManager::Get().ConvertToAbsolutePathForExternalAppForRead(*ProjectFileName));
+		}
+	}
+	
+	// -gatheronly tells UBT to update its UBT makefiles without building
+	Arguments += TEXT(" -gatheronly");
+
+	Arguments += TEXT(" -progress");
+
+	if (FRocketSupport::IsRocket())
+	{
+		Arguments += TEXT(" -rocket");
+	}
+
+	// Compile UnrealBuildTool if it doesn't exist. This can happen if we're just copying source from somewhere.
+	bool bRes = true;
+	Warn->BeginSlowTask(LOCTEXT("GatheringProjectFiles", "Gathering project files..."), true, true);
+	if(!FPaths::FileExists(GetUnrealBuildToolExecutableFilename(RootDir)))
+	{
+		Warn->StatusUpdate(0, 1, LOCTEXT("BuildingUBT", "Building UnrealBuildTool..."));
+		bRes = BuildUnrealBuildTool(RootDir, *Warn);
+	}
+	if(bRes)
+	{
+		Warn->StatusUpdate(0, 1, LOCTEXT("GatheringProjectFiles", "Gathering project files..."));
+		bRes = RunUnrealBuildTool(LOCTEXT("GatheringProjectFiles", "Gathering project files..."), RootDir, Arguments, Warn);
+	}
+	Warn->EndSlowTask();
+	return bRes;
+}
+
 bool FDesktopPlatformBase::IsUnrealBuildToolAvailable()
 {
 	// If using Rocket and the Rocket unreal build tool executable exists, then UBT is available. Otherwise check it can be built.
diff --git a/Engine/Source/Developer/DesktopPlatform/Private/DesktopPlatformBase.h b/Engine/Source/Developer/DesktopPlatform/Private/DesktopPlatformBase.h
--- a/Engine/Source/Developer/DesktopPlatform/Private/DesktopPlatformBase.h
+++ b/Engine/Source/Developer/DesktopPlatform/Private/DesktopPlatformBase.h
@@ -38,6 +38,7 @@ public:
 	virtual bool CleanGameProject(const FString& ProjectDir, FString& OutFailPath, FFeedbackContext* Warn) override;
 	virtual bool CompileGameProject(const FString& RootDir, const FString& ProjectFileName, FFeedbackContext* Warn) override;
 	virtual bool GenerateProjectFiles(const FString& RootDir, const FString& ProjectFileName, FFeedbackContext* Warn) override;
+	virtual bool GatherProjectFiles(const FString& RootDir, const FString& ProjectFileName, FFeedbackContext* Warn) override;
 	virtual bool IsUnrealBuildToolAvailable() override;
 	virtual bool InvokeUnrealBuildToolSync(const FString& InCmdLineParams, FOutputDevice &Ar, bool bSkipBuildUBT, int32& OutReturnCode, FString& OutProcOutput) override;
 	virtual FProcHandle InvokeUnrealBuildToolAsync(const FString& InCmdLineParams, FOutputDevice &Ar, void*& OutReadPipe, void*& OutWritePipe, bool bSkipBuildUBT = false) override;
diff --git a/Engine/Source/Developer/DesktopPlatform/Public/IDesktopPlatform.h b/Engine/Source/Developer/DesktopPlatform/Public/IDesktopPlatform.h
--- a/Engine/Source/Developer/DesktopPlatform/Public/IDesktopPlatform.h
+++ b/Engine/Source/Developer/DesktopPlatform/Public/IDesktopPlatform.h
@@ -329,6 +329,16 @@ public:
 	*/
 	virtual bool GenerateProjectFiles(const FString& RootDir, const FString& ProjectFileName, FFeedbackContext* Warn) = 0;
 
+	/**
+	* Gather project files for the given project (to update its UBT makefiles without performing a full project generation).
+	*
+	* @param RootDir			Engine root directory for the project to use.
+	* @param ProjectFileName	Filename of the project to update
+	* @param Warn				Feedback context to use for progress updates
+	* @return true if project files were generated successfully.
+	*/
+	virtual bool GatherProjectFiles(const FString& RootDir, const FString& ProjectFileName, FFeedbackContext* Warn) = 0;
+
 	/**
 	* Determines whether UnrealBuildTool is available
 	*
diff --git a/Engine/Source/Editor/GameProjectGeneration/Private/GameProjectUtils.cpp b/Engine/Source/Editor/GameProjectGeneration/Private/GameProjectUtils.cpp
--- a/Engine/Source/Editor/GameProjectGeneration/Private/GameProjectUtils.cpp
+++ b/Engine/Source/Editor/GameProjectGeneration/Private/GameProjectUtils.cpp
@@ -2956,12 +2956,20 @@ bool GameProjectUtils::AddCodeToProject_Internal(const FString& NewClassName, co
 		CreatedFilesForExternalAppRead.Add( IFileManager::Get().ConvertToAbsolutePathForExternalAppForRead(*CreatedFile) );
 	}
 
+	bool bGenerateProjectFiles = true;
+
 	// First see if we can avoid a full generation by adding the new files to an already open project
 	if ( bProjectHadCodeFiles && FSourceCodeNavigation::AddSourceFiles(CreatedFilesForExternalAppRead) )
 	{
-		// todo: jdale - even if this succeeds, we still need to run UBT with "-gather" to update the makefiles
+		// We successfully added the new files to the solution, but we still need to run UBT with -gather to update any UBT makefiles
+		if ( FDesktopPlatformModule::Get()->GatherProjectFiles(FPaths::RootDir(), FPaths::GetProjectFilePath(), GWarn) )
+		{
+			// We managed the gather, so we can skip running the full generate
+			bGenerateProjectFiles = false;
+		}
 	}
-	else
+	
+	if ( bGenerateProjectFiles )
 	{
 		// Generate project files if we happen to be using a project file.
 		if ( !FDesktopPlatformModule::Get()->GenerateProjectFiles(FPaths::RootDir(), FPaths::GetProjectFilePath(), GWarn) )
