Changes to building the UBT.

- Incorporates PR #934 (fix a broken path to RunXBuild.sh when building UBT for Linux) contributed by [REDACTED_PERSON] ([REDACTED_PERSON]).
- Moved more code out of #ifdef-s to make compiler see it when compiling on all platforms.

diff --git a/Engine/Source/Developer/DesktopPlatform/Private/DesktopPlatformBase.cpp b/Engine/Source/Developer/DesktopPlatform/Private/DesktopPlatformBase.cpp
--- a/Engine/Source/Developer/DesktopPlatform/Private/DesktopPlatformBase.cpp
+++ b/Engine/Source/Developer/DesktopPlatform/Private/DesktopPlatformBase.cpp
@@ -955,55 +955,65 @@ bool FDesktopPlatformBase::BuildUnrealBuildTool(const FString& RootDir, FOutputD
 
 	FString CompilerExecutableFilename;
 	FString CmdLineParams;
-#if PLATFORM_WINDOWS
-	// To build UBT for windows, we must assemble a batch file that first registers the environment variable necessary to run msbuild then run it
-	// This can not be done in a single invocation of CMD.exe because the environment variables do not transfer between subsequent commands when using the "&" syntax
-	// devenv.exe can be used to build as well but it takes several seconds to start up so it is not desirable
 
-	// First determine the appropriate vcvars batch file to launch
-	FString VCVarsBat;
+	if (PLATFORM_WINDOWS)
+	{
+		// To build UBT for windows, we must assemble a batch file that first registers the environment variable necessary to run msbuild then run it
+		// This can not be done in a single invocation of CMD.exe because the environment variables do not transfer between subsequent commands when using the "&" syntax
+		// devenv.exe can be used to build as well but it takes several seconds to start up so it is not desirable
 
-#if _MSC_VER >= 1800
-	FPlatformMisc::GetVSComnTools(12, VCVarsBat);
-#else
-	FPlatformMisc::GetVSComnTools(11, VCVarsBat);
-#endif
+		// First determine the appropriate vcvars batch file to launch
+		FString VCVarsBat;
 
-	VCVarsBat = FPaths::Combine(*VCVarsBat, L"../../VC/bin/x86_amd64/vcvarsx86_amd64.bat");
+#if PLATFORM_WINDOWS
+	#if _MSC_VER >= 1800
+		FPlatformMisc::GetVSComnTools(12, VCVarsBat);
+	#else
+		FPlatformMisc::GetVSComnTools(11, VCVarsBat);
+	#endif
+#endif // PLATFORM_WINDOWS
 
-	// Check to make sure we found one.
-	if (VCVarsBat.IsEmpty() || !FPaths::FileExists(VCVarsBat))
-	{
-		Ar.Logf(TEXT("Couldn't find %s; skipping."), *VCVarsBat);
-		return false;
-	}
+		VCVarsBat = FPaths::Combine(*VCVarsBat, L"../../VC/bin/x86_amd64/vcvarsx86_amd64.bat");
 
-	// Now make a batch file in the intermediate directory to invoke the vcvars batch then msbuild
-	FString BuildBatchFile = RootDir / TEXT("Engine/Intermediate/Build/UnrealBuildTool/BuildUBT.bat");
-	BuildBatchFile.ReplaceInline(TEXT("/"), TEXT("\\"));
+		// Check to make sure we found one.
+		if (VCVarsBat.IsEmpty() || !FPaths::FileExists(VCVarsBat))
+		{
+			Ar.Logf(TEXT("Couldn't find %s; skipping."), *VCVarsBat);
+			return false;
+		}
 
-	FString BatchFileContents;
-	BatchFileContents = FString::Printf(TEXT("call \"%s\"") LINE_TERMINATOR, *VCVarsBat);
-	BatchFileContents += FString::Printf(TEXT("msbuild /nologo /verbosity:quiet \"%s\" /property:Configuration=Development /property:Platform=AnyCPU"), *CsProjLocation);
-	FFileHelper::SaveStringToFile(BatchFileContents, *BuildBatchFile);
+		// Now make a batch file in the intermediate directory to invoke the vcvars batch then msbuild
+		FString BuildBatchFile = RootDir / TEXT("Engine/Intermediate/Build/UnrealBuildTool/BuildUBT.bat");
+		BuildBatchFile.ReplaceInline(TEXT("/"), TEXT("\\"));
 
-	TCHAR CmdExePath[MAX_PATH];
-	FPlatformMisc::GetEnvironmentVariable(TEXT("ComSpec"), CmdExePath, ARRAY_COUNT(CmdExePath));
-	CompilerExecutableFilename = CmdExePath;
+		FString BatchFileContents;
+		BatchFileContents = FString::Printf(TEXT("call \"%s\"") LINE_TERMINATOR, *VCVarsBat);
+		BatchFileContents += FString::Printf(TEXT("msbuild /nologo /verbosity:quiet \"%s\" /property:Configuration=Development /property:Platform=AnyCPU"), *CsProjLocation);
+		FFileHelper::SaveStringToFile(BatchFileContents, *BuildBatchFile);
 
-	CmdLineParams = FString::Printf(TEXT("/c \"%s\""), *BuildBatchFile);
-#elif PLATFORM_MAC
-	FString ScriptPath = FPaths::ConvertRelativePathToFull(RootDir / TEXT("Engine/Build/BatchFiles/Mac/RunXBuild.sh"));
-	CompilerExecutableFilename = TEXT("/bin/sh");
-	CmdLineParams = FString::Printf(TEXT("\"%s\" /property:Configuration=Development %s"), *ScriptPath, *CsProjLocation);
-#elif PLATFORM_LINUX
-	FString ScriptPath = FPaths::ConvertRelativePathToFull(RootDir / TEXT("Build/BatchFiles/Linux/RunXBuild.sh"));
-	CompilerExecutableFilename = TEXT("/bin/bash");
-	CmdLineParams = FString::Printf(TEXT("\"%s\" /property:Configuration=Development /property:TargetFrameworkVersion=v4.0 %s"), *ScriptPath, *CsProjLocation);
-#else
-	Ar.Log(TEXT("Unknown platform, unable to build UnrealBuildTool."));
-	return false;
-#endif
+		TCHAR CmdExePath[MAX_PATH];
+		FPlatformMisc::GetEnvironmentVariable(TEXT("ComSpec"), CmdExePath, ARRAY_COUNT(CmdExePath));
+		CompilerExecutableFilename = CmdExePath;
+
+		CmdLineParams = FString::Printf(TEXT("/c \"%s\""), *BuildBatchFile);
+	}
+	else if (PLATFORM_MAC)
+	{
+		FString ScriptPath = FPaths::ConvertRelativePathToFull(RootDir / TEXT("Engine/Build/BatchFiles/Mac/RunXBuild.sh"));
+		CompilerExecutableFilename = TEXT("/bin/sh");
+		CmdLineParams = FString::Printf(TEXT("\"%s\" /property:Configuration=Development %s"), *ScriptPath, *CsProjLocation);
+	}
+	else if (PLATFORM_LINUX)
+	{
+		FString ScriptPath = FPaths::ConvertRelativePathToFull(RootDir / TEXT("Engine/Build/BatchFiles/Linux/RunXBuild.sh"));
+		CompilerExecutableFilename = TEXT("/bin/bash");
+		CmdLineParams = FString::Printf(TEXT("\"%s\" /property:Configuration=Development /property:TargetFrameworkVersion=v4.0 %s"), *ScriptPath, *CsProjLocation);
+	}
+	else
+	{
+		Ar.Log(TEXT("Unknown platform, unable to build UnrealBuildTool."));
+		return false;
+	}
 
 	// Spawn the compiler
 	Ar.Logf(TEXT("Running: %s %s"), *CompilerExecutableFilename, *CmdLineParams);
