Make sure initial compile forwards -Rocket flag to UBT when running from source with the -Rocket argument.

diff --git a/Engine/Source/Developer/DesktopPlatform/Private/DesktopPlatformBase.cpp b/Engine/Source/Developer/DesktopPlatform/Private/DesktopPlatformBase.cpp
--- a/Engine/Source/Developer/DesktopPlatform/Private/DesktopPlatformBase.cpp
+++ b/Engine/Source/Developer/DesktopPlatform/Private/DesktopPlatformBase.cpp
@@ -415,7 +415,7 @@ bool FDesktopPlatformBase::CompileGameProject(const FString& RootDir, const FStr
 	}
 
 	// Append the Rocket flag
-	if(!IsSourceDistribution(RootDir))
+	if(!IsSourceDistribution(RootDir) || FRocketSupport::IsRocket())
 	{
 		Arguments += TEXT(" -rocket");
 	}
