Merging using UE4-To-UE4-LauncherDev

diff --git a/Engine/Source/Developer/DesktopPlatform/Private/DesktopPlatformBase.cpp b/Engine/Source/Developer/DesktopPlatform/Private/DesktopPlatformBase.cpp
--- a/Engine/Source/Developer/DesktopPlatform/Private/DesktopPlatformBase.cpp
+++ b/Engine/Source/Developer/DesktopPlatform/Private/DesktopPlatformBase.cpp
@@ -864,18 +864,18 @@ void FDesktopPlatformBase::GetProjectBuildProducts(const FString& ProjectDir, TA
 	}
 }
 
-bool FDesktopPlatformBase::EnumerateProjectsKnownByEngine(const FString &Identifier, bool bIncludeNativeProjects, TArray<FString> &OutProjectFileNames)
+FString FDesktopPlatformBase::GetEngineSavedConfigDirectory(const FString& Identifier)
 {
 	// Get the engine root directory
 	FString RootDir;
-	if(!GetEngineRootDirFromIdentifier(Identifier, RootDir))
+	if (!GetEngineRootDirFromIdentifier(Identifier, RootDir))
 	{
-		return false;
+		return FString();
 	}
 
 	// Get the path to the game agnostic settings
 	FString UserDir;
-	if(IsStockEngineRelease(Identifier))
+	if (IsStockEngineRelease(Identifier))
 	{
 		UserDir = FPaths::Combine(FPlatformProcess::UserSettingsDir(), *FString(EPIC_PRODUCT_IDENTIFIER), *Identifier);
 	}
@@ -885,7 +885,24 @@ bool FDesktopPlatformBase::EnumerateProjectsKnownByEngine(const FString &Identif
 	}
 
 	// Get the game agnostic config dir
-	FString GameAgnosticConfigDir = UserDir / TEXT("Saved/Config") / ANSI_TO_TCHAR(FPlatformProperties::PlatformName());
+	return UserDir / TEXT("Saved/Config") / ANSI_TO_TCHAR(FPlatformProperties::PlatformName());
+}
+
+bool FDesktopPlatformBase::EnumerateProjectsKnownByEngine(const FString &Identifier, bool bIncludeNativeProjects, TArray<FString> &OutProjectFileNames)
+{
+	// Get the engine root directory
+	FString RootDir;
+	if (!GetEngineRootDirFromIdentifier(Identifier, RootDir))
+	{
+		return false;
+	}
+
+	FString GameAgnosticConfigDir = GetEngineSavedConfigDirectory(Identifier);
+
+	if (GameAgnosticConfigDir.Len() == 0)
+	{
+		return false;
+	}
 
 	// Find all the created project directories. Start with the default project creation path.
 	TArray<FString> SearchDirectories;
diff --git a/Engine/Source/Developer/DesktopPlatform/Private/DesktopPlatformBase.h b/Engine/Source/Developer/DesktopPlatform/Private/DesktopPlatformBase.h
--- a/Engine/Source/Developer/DesktopPlatform/Private/DesktopPlatformBase.h
+++ b/Engine/Source/Developer/DesktopPlatform/Private/DesktopPlatformBase.h
@@ -23,6 +23,7 @@ public:
 
 	virtual bool GetDefaultEngineIdentifier(FString &OutIdentifier) override;
 	virtual bool GetDefaultEngineRootDir(FString &OutRootDir) override;
+	virtual FString GetEngineSavedConfigDirectory(const FString& Identifier) override;
 	virtual bool IsPreferredEngineIdentifier(const FString &Identifier, const FString &OtherIdentifier) override;
 
 	virtual bool IsStockEngineRelease(const FString &Identifier) override;
diff --git a/Engine/Source/Developer/DesktopPlatform/Public/IDesktopPlatform.h b/Engine/Source/Developer/DesktopPlatform/Public/IDesktopPlatform.h
--- a/Engine/Source/Developer/DesktopPlatform/Public/IDesktopPlatform.h
+++ b/Engine/Source/Developer/DesktopPlatform/Public/IDesktopPlatform.h
@@ -235,6 +235,14 @@ public:
 	*/
 	virtual bool GetDefaultEngineRootDir(FString &OutRootDir) = 0;
 
+	/**
+	* Gets the root directory for the engine's saved config files
+	*
+	* @param	Identifier		Identifier for the engine
+	* @return	The absoluate path to the engines Saved/Config directory
+	*/
+	virtual FString GetEngineSavedConfigDirectory(const FString& Identifier) = 0;
+
 	/**
 	* Checks if the given engine identifier is for an stock engine release.
 	*
