LND: Fix FBX import (UE-13581).

- Lumps all filters into one to work around current limitiations of the third party lib.

diff --git a/Engine/Source/Developer/DesktopPlatform/Private/Linux/DesktopPlatformLinux.cpp b/Engine/Source/Developer/DesktopPlatform/Private/Linux/DesktopPlatformLinux.cpp
--- a/Engine/Source/Developer/DesktopPlatform/Private/Linux/DesktopPlatformLinux.cpp
+++ b/Engine/Source/Developer/DesktopPlatform/Private/Linux/DesktopPlatformLinux.cpp
@@ -115,54 +115,45 @@ bool FDesktopPlatformLinux::FileDialogShared(bool bSave, const void* ParentWindo
 
 	hints.WindowTitle = TCHAR_TO_UTF8(*DialogTitle);
 
-	// Convert the "|" delimited list of filetypes to NULL delimited then add a second NULL character to indicate the end of the list
-	TCHAR FileTypeStr[MAX_FILETYPES_STR];
-	TCHAR* FileTypesPtr = NULL;
-	const int32 FileTypesLen = FileTypes.Len();
-
-	// Nicely formatted file types for lookup later and suitable to append to filenames without extensions
-	TArray<FString> CleanExtensionList;
-
-	// The strings must be in pairs for windows.
-	// It is formatted as follows: Pair1String1|Pair1String2|Pair2String1|Pair2String2
-	// where the second string in the pair is the extension.  To get the clean extensions we only care about the second string in the pair
-	TArray<FString> UnformattedExtensions;
-	FileTypes.ParseIntoArray( UnformattedExtensions, TEXT("|"), true );
-	for( int32 ExtensionIndex = 1; ExtensionIndex < UnformattedExtensions.Num(); ExtensionIndex += 2)
+	FString AllExtensionsSpaceDelim;	// a string like "*.cpp *.h *.c"
+
+	// The strings will come in pairs, formatted as follows: Pair1String1|Pair1String2|Pair2String1|Pair2String2
+	// where the second string in the pair is the extension(s)
+	TArray<FString> Filters;
+	FileTypes.ParseIntoArray(Filters, TEXT("|"), true);
+	for( int32 ExtensionIndex = 1; ExtensionIndex < Filters.Num(); ExtensionIndex += 2)
 	{
-		const FString& Extension = UnformattedExtensions[ExtensionIndex];
-		// Assume the user typed in an extension or doesn't want one when using the *.* extension. We can't determine what extension they want in that case
-		if( Extension != TEXT("*.*") )
+		const FString& Extensions = Filters[ExtensionIndex];
+
+		TArray<FString> ExtensionsArray;
+
+		// Extension can be either *.jpg or *.jpg;*.png -> in the latter case split at ';'
+		int32 UnusedIndex;
+		if (Extensions.FindChar(TEXT(';'), UnusedIndex))
 		{
-			// Add to the clean extension list, first removing the * wildcard from the extension
-			int32 WildCardIndex = Extension.Find( TEXT("*") );
-			CleanExtensionList.Add( WildCardIndex != INDEX_NONE ? Extension.RightChop( WildCardIndex+1 ) : Extension );
+			Extensions.ParseIntoArray(ExtensionsArray, TEXT(";"), true);
+		}
+		else
+		{
+			ExtensionsArray.Add(Extensions);	// just a single extension
 		}
-	}
-
-	if (FileTypesLen > 0 && FileTypesLen - 1 < MAX_FILETYPES_STR)
-	{
-		FileTypesPtr = FileTypeStr;
-		FCString::Strcpy(FileTypeStr, MAX_FILETYPES_STR, *FileTypes.Replace(TEXT(";"), TEXT(" ")));
 
-		TCHAR* Pos = FileTypeStr;
-		while( Pos[0] != 0 )
+		for (const FString& Extension : ExtensionsArray)
 		{
-			if ( Pos[0] == '|' )
+			if (AllExtensionsSpaceDelim.Find(Extension, ESearchCase::CaseSensitive) == INDEX_NONE)
 			{
-				Pos[0] = 0;
+				AllExtensionsSpaceDelim += Extension;
+				AllExtensionsSpaceDelim += TEXT(" ");
 			}
-
-			Pos++;
 		}
-
-		// Add two trailing NULL characters to indicate the end of the list
-		FileTypeStr[FileTypesLen] = 0;
-		FileTypeStr[FileTypesLen + 1] = 0;
 	}
 
+	FString AllExtensionsLumpedTogether(TEXT("All applicable ("));
+	AllExtensionsLumpedTogether += AllExtensionsSpaceDelim;
+	AllExtensionsLumpedTogether += TEXT(")");
+
 	char FileTypesBuf[MAX_FILETYPES_STR * 2] = {0,};
-	FTCHARToUTF8_Convert::Convert(FileTypesBuf, sizeof(FileTypesBuf), FileTypeStr, FileTypesLen);
+	FTCHARToUTF8_Convert::Convert(FileTypesBuf, sizeof(FileTypesBuf), *AllExtensionsSpaceDelim, AllExtensionsSpaceDelim.Len());
 	hints.NameFilter = FileTypesBuf;
 
 	char DefPathBuf[MAX_FILENAME_STR * 2] = {0,};
@@ -199,7 +190,7 @@ bool FDesktopPlatformLinux::FileDialogShared(bool bSave, const void* ParentWindo
 			UE_LOG(LogDesktopPlatform, Warning, TEXT("FileDialogShared Selected Files: %d"), result->count);
 			for(int i = 0;i < result->count;++i) {
 				//Filename = FUTF8ToTCHAR(result->selection[i], MAX_FILENAME_STR).Get();
-				Filename = UTF8_TO_TCHAR(result->selection[0]);
+				Filename = UTF8_TO_TCHAR(result->selection[i]);
 				//new(OutFilenames) FString(Filename);
 				OutFilenames.Add(Filename);
 				Filename = IFileManager::Get().ConvertToRelativePath(*Filename);
