Fixed C++ variable shadowing in Windows-specific code
- Two actual bugs were fixed (GitSourceControlUtils.cpp, VisualStudioSourceCodeAccessor.h)
- (These errors were found by using the Clang compiler on Windows platform)

diff --git a/Engine/Plugins/Developer/GitSourceControl/Source/GitSourceControl/Private/GitSourceControlUtils.cpp b/Engine/Plugins/Developer/GitSourceControl/Source/GitSourceControl/Private/GitSourceControlUtils.cpp
--- a/Engine/Plugins/Developer/GitSourceControl/Source/GitSourceControl/Private/GitSourceControlUtils.cpp
+++ b/Engine/Plugins/Developer/GitSourceControl/Source/GitSourceControl/Private/GitSourceControlUtils.cpp
@@ -116,7 +116,7 @@ FString FindGitBinaryPath()
 	// 2) Else, look for the version of Git bundled with SmartGit "Installer with JRE"
 	if(!bFound)
 	{
-		FString GitBinaryPath(TEXT("C:/Program Files (x86)/SmartGit/bin/git.exe"));
+		GitBinaryPath = TEXT("C:/Program Files (x86)/SmartGit/bin/git.exe");
 		bFound = CheckGitAvailability(GitBinaryPath);
 	}
 
diff --git a/Engine/Plugins/Developer/VisualStudioSourceCodeAccess/Source/VisualStudioSourceCodeAccess/Private/VisualStudioSourceCodeAccessor.h b/Engine/Plugins/Developer/VisualStudioSourceCodeAccess/Source/VisualStudioSourceCodeAccess/Private/VisualStudioSourceCodeAccessor.h
--- a/Engine/Plugins/Developer/VisualStudioSourceCodeAccess/Source/VisualStudioSourceCodeAccess/Private/VisualStudioSourceCodeAccessor.h
+++ b/Engine/Plugins/Developer/VisualStudioSourceCodeAccess/Source/VisualStudioSourceCodeAccess/Private/VisualStudioSourceCodeAccessor.h
@@ -51,10 +51,10 @@ private:
 		{
 		}
 
-		FileOpenRequest(const FString& FullPath, int32 LineNumber, int32 ColumnNumber)
-			: FullPath(FullPath)
-			, LineNumber(LineNumber)
-			, ColumnNumber(ColumnNumber)
+		FileOpenRequest(const FString& InitFullPath, int32 InitLineNumber, int32 InitColumnNumber)
+			: FullPath(InitFullPath)
+			, LineNumber(InitLineNumber)
+			, ColumnNumber(InitColumnNumber)
 		{
 		}
 	};
diff --git a/Engine/Plugins/Media/WmfMedia/Source/WmfMedia/Private/Player/WmfMediaPlayer.cpp b/Engine/Plugins/Media/WmfMedia/Source/WmfMedia/Private/Player/WmfMediaPlayer.cpp
--- a/Engine/Plugins/Media/WmfMedia/Source/WmfMedia/Private/Player/WmfMediaPlayer.cpp
+++ b/Engine/Plugins/Media/WmfMedia/Source/WmfMedia/Private/Player/WmfMediaPlayer.cpp
@@ -261,7 +261,7 @@ bool FWmfMediaPlayer::SetRate( float Rate )
 /* FWmfMediaPlayer implementation
  *****************************************************************************/
 
-void FWmfMediaPlayer::AddStreamToTopology( uint32 StreamIndex, IMFTopology* Topology, IMFPresentationDescriptor* PresentationDescriptor, IMFMediaSource* MediaSource )
+void FWmfMediaPlayer::AddStreamToTopology( uint32 StreamIndex, IMFTopology* Topology, IMFPresentationDescriptor* PresentationDescriptor, IMFMediaSource* MediaSourceObject )
 {
 	// get stream descriptor
 	TComPtr<IMFStreamDescriptor> StreamDescriptor;
@@ -357,7 +357,7 @@ void FWmfMediaPlayer::AddStreamToTopology( uint32 StreamIndex, IMFTopology* Topo
 		TComPtr<IMFTopologyNode> SourceNode;
 		{
 			if (FAILED(::MFCreateTopologyNode(MF_TOPOLOGY_SOURCESTREAM_NODE, &SourceNode)) ||
-				FAILED(SourceNode->SetUnknown(MF_TOPONODE_SOURCE, MediaSource)) ||
+				FAILED(SourceNode->SetUnknown(MF_TOPONODE_SOURCE, MediaSourceObject)) ||
 				FAILED(SourceNode->SetUnknown(MF_TOPONODE_PRESENTATION_DESCRIPTOR, PresentationDescriptor)) ||
 				FAILED(SourceNode->SetUnknown(MF_TOPONODE_STREAM_DESCRIPTOR, StreamDescriptor)) ||
 				FAILED(Topology->AddNode(SourceNode)))
diff --git a/Engine/Plugins/Media/WmfMedia/Source/WmfMedia/Private/Player/WmfMediaPlayer.h b/Engine/Plugins/Media/WmfMedia/Source/WmfMedia/Private/Player/WmfMediaPlayer.h
--- a/Engine/Plugins/Media/WmfMedia/Source/WmfMedia/Private/Player/WmfMediaPlayer.h
+++ b/Engine/Plugins/Media/WmfMedia/Source/WmfMedia/Private/Player/WmfMediaPlayer.h
@@ -70,10 +70,10 @@ protected:
 	 * @param StreamIndex The index number of the stream in the presentation descriptor.
 	 * @param Topology The topology to add the stream to.
 	 * @param PresentationDescriptor The presentation descriptor object.
-	 * @param MediaSource The media source object.
+	 * @param MediaSourceObject The media source object.
 	 * @return true on success, false otherwise.
 	 */
-	void AddStreamToTopology( uint32 StreamIndex, IMFTopology* Topology, IMFPresentationDescriptor* PresentationDescriptor, IMFMediaSource* MediaSource );
+	void AddStreamToTopology( uint32 StreamIndex, IMFTopology* Topology, IMFPresentationDescriptor* PresentationDescriptor, IMFMediaSource* MediaSourceObject );
 
 	/**
 	 * Initializes the media session for the given media source.
diff --git a/Engine/Plugins/Runtime/LeapMotionController/Source/LeapMotionController/Private/LeapMotionBoneActor.cpp b/Engine/Plugins/Runtime/LeapMotionController/Source/LeapMotionController/Private/LeapMotionBoneActor.cpp
--- a/Engine/Plugins/Runtime/LeapMotionController/Source/LeapMotionController/Private/LeapMotionBoneActor.cpp
+++ b/Engine/Plugins/Runtime/LeapMotionController/Source/LeapMotionController/Private/LeapMotionBoneActor.cpp
@@ -103,8 +103,8 @@ void ALeapMotionBoneActor::Init(ELeapBone LeapBone, float Scale, float Width, fl
 				Capsule->SetPhysicsMaxAngularVelocity(7200.0f);
 			}
 
-			UStaticMeshComponent* StaticMeshComponent = Cast<UStaticMeshComponent>(c);
-			UStaticMesh* StaticMesh = StaticMeshComponent ? StaticMeshComponent->StaticMesh : nullptr;
+			UStaticMeshComponent* NonRootStaticMeshComponent = Cast<UStaticMeshComponent>(c);
+			UStaticMesh* StaticMesh = NonRootStaticMeshComponent ? NonRootStaticMeshComponent->StaticMesh : nullptr;
 			if (StaticMesh)
 			{
 				FBoxSphereBounds Bounds = StaticMesh->GetBounds();
@@ -120,9 +120,9 @@ void ALeapMotionBoneActor::Init(ELeapBone LeapBone, float Scale, float Width, fl
 				{
 					ScaleX = 0.06125f * Scale * Width / Bounds.BoxExtent.X;
 				}
-				StaticMeshComponent->SetRelativeScale3D(FVector(ScaleX, ScaleY, ScaleZ));
-				StaticMeshComponent->AddLocalOffset(BoneOffset);
-				StaticMeshComponent->SetPhysicsMaxAngularVelocity(7200.0f);
+				NonRootStaticMeshComponent->SetRelativeScale3D(FVector(ScaleX, ScaleY, ScaleZ));
+				NonRootStaticMeshComponent->AddLocalOffset(BoneOffset);
+				NonRootStaticMeshComponent->SetPhysicsMaxAngularVelocity(7200.0f);
 			}
 		}
 	}
diff --git a/Engine/Source/Developer/DesktopPlatform/Private/Windows/DesktopPlatformWindows.cpp b/Engine/Source/Developer/DesktopPlatform/Private/Windows/DesktopPlatformWindows.cpp
--- a/Engine/Source/Developer/DesktopPlatform/Private/Windows/DesktopPlatformWindows.cpp
+++ b/Engine/Source/Developer/DesktopPlatform/Private/Windows/DesktopPlatformWindows.cpp
@@ -366,19 +366,19 @@ bool FDesktopPlatformWindows::FileDialogShared(bool bSave, const void* ParentWin
 		FString Extension = CleanExtensionList.IsValidIndex( OutFilterIndex ) ? CleanExtensionList[OutFilterIndex] : TEXT("");
 
 		// Make sure all filenames gathered have their paths normalized and proper extensions added
-		for ( auto FilenameIt = OutFilenames.CreateIterator(); FilenameIt; ++FilenameIt )
+		for ( auto OutFilenameIt = OutFilenames.CreateIterator(); OutFilenameIt; ++OutFilenameIt )
 		{
-			FString& Filename = *FilenameIt;
+			FString& OutFilename = *OutFilenameIt;
 			
-			Filename = IFileManager::Get().ConvertToRelativePath(*Filename);
+			OutFilename = IFileManager::Get().ConvertToRelativePath(*OutFilename);
 
-			if( FPaths::GetExtension(Filename).IsEmpty() && !Extension.IsEmpty() )
+			if( FPaths::GetExtension(OutFilename).IsEmpty() && !Extension.IsEmpty() )
 			{
 				// filename does not have an extension. Add an extension based on the filter that the user chose in the dialog
-				Filename += Extension;
+				OutFilename += Extension;
 			}
 
-			FPaths::NormalizeFilename(Filename);
+			FPaths::NormalizeFilename(OutFilename);
 		}
 	}
 	else
diff --git a/Engine/Source/Developer/DesktopPlatform/Private/Windows/WindowsNativeFeedbackContext.cpp b/Engine/Source/Developer/DesktopPlatform/Private/Windows/WindowsNativeFeedbackContext.cpp
--- a/Engine/Source/Developer/DesktopPlatform/Private/Windows/WindowsNativeFeedbackContext.cpp
+++ b/Engine/Source/Developer/DesktopPlatform/Private/Windows/WindowsNativeFeedbackContext.cpp
@@ -99,10 +99,10 @@ bool FWindowsNativeFeedbackContext::ReceivedUserCancel()
 	return bReceivedUserCancel;
 }
 
-void FWindowsNativeFeedbackContext::StartSlowTask( const FText& Task, bool bShowCancelButton )
+void FWindowsNativeFeedbackContext::StartSlowTask( const FText& Task, bool bShouldShowCancelButton )
 {
-	FFeedbackContext::StartSlowTask( Task, bShowCancelButton );
-	CreateSlowTaskWindow(Task, bShowCancelButton);
+	FFeedbackContext::StartSlowTask( Task, bShouldShowCancelButton );
+	CreateSlowTaskWindow(Task, bShouldShowCancelButton);
 }
 void FWindowsNativeFeedbackContext::FinalizeSlowTask( )
 {
@@ -178,16 +178,16 @@ DWORD FWindowsNativeFeedbackContext::SlowTaskThreadProc(void* ThreadParam)
 {
 	FWindowsNativeFeedbackContext* Context = (FWindowsNativeFeedbackContext*)ThreadParam;
 
-	HINSTANCE hInstance = (HINSTANCE)GetModuleHandle(NULL);
+	HINSTANCE HInstance = (HINSTANCE)GetModuleHandle(NULL);
 
 	WNDCLASSEX WndClassEx;
 	ZeroMemory(&WndClassEx, sizeof(WndClassEx));
 	WndClassEx.cbSize = sizeof(WndClassEx);
 	WndClassEx.style = CS_HREDRAW | CS_VREDRAW | (Context->bShowCancelButton? 0 : CS_NOCLOSE);
 	WndClassEx.lpfnWndProc = &SlowTaskWindowProc;
-	WndClassEx.hIcon = LoadIcon(hInstance, MAKEINTRESOURCE(FWindowsPlatformMisc::GetAppIcon()));
+	WndClassEx.hIcon = LoadIcon(HInstance, MAKEINTRESOURCE(FWindowsPlatformMisc::GetAppIcon()));
 	WndClassEx.hCursor = LoadCursor(NULL, IDC_ARROW);
-	WndClassEx.hInstance = hInstance;
+	WndClassEx.hInstance = HInstance;
 	WndClassEx.hbrBackground = (HBRUSH)(COLOR_BTNFACE + 1);
 	WndClassEx.lpszClassName = TEXT("FFeedbackContextWindows");
 	ATOM WndClassAtom = RegisterClassEx(&WndClassEx);
@@ -226,20 +226,20 @@ DWORD FWindowsNativeFeedbackContext::SlowTaskThreadProc(void* ThreadParam)
 	AdjustWindowRectEx(&WindowRect, WindowStyle, 0, 0);
 
 	const TCHAR* WindowClassName = MAKEINTATOM( WndClassAtom );
-	HWND hWnd = CreateWindow(WindowClassName, TEXT("Unreal Engine"), WindowStyle, WindowRect.left, WindowRect.top, WindowRect.right - WindowRect.left, WindowRect.bottom - WindowRect.top, NULL, NULL, hInstance, NULL);
+	HWND hWnd = CreateWindow(WindowClassName, TEXT("Unreal Engine"), WindowStyle, WindowRect.left, WindowRect.top, WindowRect.right - WindowRect.left, WindowRect.bottom - WindowRect.top, NULL, NULL, HInstance, NULL);
 	SetWindowLongPtr(hWnd, GWLP_USERDATA, (LONG_PTR)&Params);
 	SendMessage(hWnd, WM_SETFONT, (WPARAM)hFont, 0);
 
-	HWND hWndOpenLog = CreateWindow(WC_BUTTON, TEXT("Show log"), BS_CENTER | BS_VCENTER | BS_PUSHBUTTON | BS_TEXT | WS_CHILD | WS_VISIBLE, 10, 10, 10, 10, hWnd, (HMENU)ShowLogCtlId, hInstance, NULL);
+	HWND hWndOpenLog = CreateWindow(WC_BUTTON, TEXT("Show log"), BS_CENTER | BS_VCENTER | BS_PUSHBUTTON | BS_TEXT | WS_CHILD | WS_VISIBLE, 10, 10, 10, 10, hWnd, (HMENU)ShowLogCtlId, HInstance, NULL);
 	SendMessage(hWndOpenLog, WM_SETFONT, (WPARAM)hFont, 0);
 
-	HWND hWndStatus = CreateWindow(WC_STATIC, TEXT(""), SS_CENTER | WS_CHILD | WS_VISIBLE, 10, 10, 10, 10, hWnd, (HMENU)StatusCtlId, hInstance, NULL);
+	HWND hWndStatus = CreateWindow(WC_STATIC, TEXT(""), SS_CENTER | WS_CHILD | WS_VISIBLE, 10, 10, 10, 10, hWnd, (HMENU)StatusCtlId, HInstance, NULL);
 	SendMessage(hWndStatus, WM_SETFONT, (WPARAM)hFont, 0);
 
-	HWND hWndProgress = CreateWindowEx(0, PROGRESS_CLASS, TEXT(""), WS_CHILD | WS_VISIBLE, 10, 10, 10, 10, hWnd, (HMENU)ProgressCtlId, hInstance, NULL);
+	HWND hWndProgress = CreateWindowEx(0, PROGRESS_CLASS, TEXT(""), WS_CHILD | WS_VISIBLE, 10, 10, 10, 10, hWnd, (HMENU)ProgressCtlId, HInstance, NULL);
 	SendMessage(hWndProgress, PBM_SETRANGE32, 0, 1000);
 
-	HWND hWndLogOutput = CreateWindowEx(WS_EX_STATICEDGE, WC_EDIT, TEXT(""), ES_MULTILINE | ES_READONLY | WS_HSCROLL | WS_VSCROLL | WS_CHILD | WS_VISIBLE, 10, 10, 10, 10, hWnd, (HMENU)LogOutputCtlId, hInstance, NULL);
+	HWND hWndLogOutput = CreateWindowEx(WS_EX_STATICEDGE, WC_EDIT, TEXT(""), ES_MULTILINE | ES_READONLY | WS_HSCROLL | WS_VSCROLL | WS_CHILD | WS_VISIBLE, 10, 10, 10, 10, hWnd, (HMENU)LogOutputCtlId, HInstance, NULL);
 	SendMessage(hWndLogOutput, WM_SETFONT, (WPARAM)hLogFont, 0);
 
 	LayoutControls(hWnd, &Params);
@@ -294,7 +294,7 @@ DWORD FWindowsNativeFeedbackContext::SlowTaskThreadProc(void* ThreadParam)
 	DestroyWindow(hWnd);
 	DeleteObject(hLogFont);
 	DeleteObject(hFont);
-	UnregisterClass(WindowClassName, hInstance);
+	UnregisterClass(WindowClassName, HInstance);
 
 	return 0;
 }
diff --git a/Engine/Source/Developer/DesktopPlatform/Private/Windows/WindowsNativeFeedbackContext.h b/Engine/Source/Developer/DesktopPlatform/Private/Windows/WindowsNativeFeedbackContext.h
--- a/Engine/Source/Developer/DesktopPlatform/Private/Windows/WindowsNativeFeedbackContext.h
+++ b/Engine/Source/Developer/DesktopPlatform/Private/Windows/WindowsNativeFeedbackContext.h
@@ -21,7 +21,7 @@ public:
 	VARARG_BODY( bool, YesNof, const TCHAR*, VARARG_NONE );
 
 	virtual bool ReceivedUserCancel() override;
-	virtual void StartSlowTask( const FText& Task, bool bShowCancelButton=false ) override;
+	virtual void StartSlowTask( const FText& Task, bool bShouldShowCancelButton=false ) override;
 	virtual void FinalizeSlowTask( ) override;
 	virtual void ProgressReported( const float TotalProgressInterp, FText DisplayMessage ) override;
 
diff --git a/Engine/Source/Editor/GameProjectGeneration/Private/SNewProjectWizard.cpp b/Engine/Source/Editor/GameProjectGeneration/Private/SNewProjectWizard.cpp
--- a/Engine/Source/Editor/GameProjectGeneration/Private/SNewProjectWizard.cpp
+++ b/Engine/Source/Editor/GameProjectGeneration/Private/SNewProjectWizard.cpp
@@ -24,7 +24,7 @@ FName SNewProjectWizard::TemplatePageName = TEXT("Template");
 FName SNewProjectWizard::NameAndLocationPageName = TEXT("NameAndLocation");
 
 
-namespace
+namespace NewProjectWizardDefs
 {
 	const float ThumbnailSize = 64.f, ThumbnailPadding = 5.f;
 	const float ItemWidth = ThumbnailSize + 2*ThumbnailPadding;
@@ -250,11 +250,11 @@ public:
 				+ SVerticalBox::Slot()
 				.AutoHeight()
 				.HAlign(HAlign_Center)
-				.Padding(ThumbnailPadding)
+				.Padding(NewProjectWizardDefs::ThumbnailPadding)
 				[
 					SNew(SBox)
-					.WidthOverride(ThumbnailSize)
-					.HeightOverride(ThumbnailSize)
+					.WidthOverride( NewProjectWizardDefs::ThumbnailSize )
+					.HeightOverride( NewProjectWizardDefs::ThumbnailSize )
 					[
 						SNew(SImage)
 						.Image(this, &STemplateTile::GetThumbnail)
@@ -265,10 +265,10 @@ public:
 				+ SVerticalBox::Slot()
 				.HAlign(HAlign_Center)
 				.VAlign(VAlign_Top)
-				.Padding(FMargin(ThumbnailPadding, 0))
+				.Padding(FMargin(NewProjectWizardDefs::ThumbnailPadding, 0))
 				[
 					SNew(STextBlock)
-					.WrapTextAt(ThumbnailSize)
+					.WrapTextAt( NewProjectWizardDefs::ThumbnailSize )
 					.Justification(ETextJustify::Center)
 					.LineBreakPolicy(FBreakIterator::CreateCamelCaseBreakIterator())
 					//.HighlightText(this, &SNewProjectWizard::GetItemHighlightText)
@@ -318,8 +318,8 @@ void SNewProjectWizard::Construct( const FArguments& InArgs )
 	.SelectionMode(ESelectionMode::Single)
 	.ClearSelectionOnClick(false)
 	.OnGenerateTile_Static(&STemplateTile::BuildTile)
-	.ItemHeight(ItemHeight)
-	.ItemWidth(ItemWidth)
+	.ItemHeight( NewProjectWizardDefs::ItemHeight )
+	.ItemWidth( NewProjectWizardDefs::ItemWidth )
 	.OnMouseButtonDoubleClick(this, &SNewProjectWizard::HandleTemplateListViewDoubleClick)
 	.OnSelectionChanged(this, &SNewProjectWizard::HandleTemplateListViewSelectionChanged);
 
diff --git a/Engine/Source/Editor/UnrealEd/Private/Factories/TTFontImport.cpp b/Engine/Source/Editor/UnrealEd/Private/Factories/TTFontImport.cpp
--- a/Engine/Source/Editor/UnrealEd/Private/Factories/TTFontImport.cpp
+++ b/Engine/Source/Editor/UnrealEd/Private/Factories/TTFontImport.cpp
@@ -1277,8 +1277,6 @@ bool UTrueTypeFontFactory::CreateFontTexture(
 				CurrentTexture++;
 
 				// blank out DC
-				HBRUSH Black = CreateSolidBrush(0x00000000);
-				RECT r = {0, 0, BitmapWidth, BitmapHeight};
 				FillRect( DCHandle, &r, Black );
 
 				X=BitmapPaddingX;
diff --git a/Engine/Source/Runtime/Core/Private/Windows/WindowsPlatformSurvey.cpp b/Engine/Source/Runtime/Core/Private/Windows/WindowsPlatformSurvey.cpp
--- a/Engine/Source/Runtime/Core/Private/Windows/WindowsPlatformSurvey.cpp
+++ b/Engine/Source/Runtime/Core/Private/Windows/WindowsPlatformSurvey.cpp
@@ -669,7 +669,7 @@ bool FWindowsPlatformSurvey::GetSubComponentIndex( IProvideWinSATResultsInfo* Wi
 	return bSuccess;
 }
 
-void FWindowsPlatformSurvey::GetOSVersionLabels(const SYSTEM_INFO& SystemInfo, FHardwareSurveyResults& Results)
+void FWindowsPlatformSurvey::GetOSVersionLabels(const SYSTEM_INFO& SystemInfo, FHardwareSurveyResults& OutResults)
 {
 	FString OSVersionLabel;
 	FString OSSubVersionLabel;
@@ -677,9 +677,9 @@ void FWindowsPlatformSurvey::GetOSVersionLabels(const SYSTEM_INFO& SystemInfo, F
 
 	if( ErrorCode & FWindowsOSVersionHelper::ERROR_GETPRODUCTINFO_FAILED )
 	{
-		Results.ErrorCount++;
-		WriteFStringToResults( Results.LastSurveyError, TEXT( "Failed to get GetProductInfo() function from GetProcAddress()." ) );
-		WriteFStringToResults( Results.LastSurveyErrorDetail, TEXT( "" ) );
+		OutResults.ErrorCount++;
+		WriteFStringToResults( OutResults.LastSurveyError, TEXT( "Failed to get GetProductInfo() function from GetProcAddress()." ) );
+		WriteFStringToResults( OutResults.LastSurveyErrorDetail, TEXT( "" ) );
 	}
 
 	if( ErrorCode & FWindowsOSVersionHelper::ERROR_UNKNOWNVERSION )
@@ -693,22 +693,22 @@ void FWindowsPlatformSurvey::GetOSVersionLabels(const SYSTEM_INFO& SystemInfo, F
 #pragma warning(pop)
 
 		UE_LOG( LogWindows, Warning, TEXT( "FWindowsPlatformSurvey::GetOSVersionLabel() unknown Windows version info from GetVersionEx()" ) );
-		Results.ErrorCount++;
-		WriteFStringToResults( Results.LastSurveyError, TEXT( "GetVersionEx() returned unknown version" ) );
-		WriteFStringToResults( Results.LastSurveyErrorDetail, FString::Printf( TEXT( "dwMajorVersion: %d  dwMinorVersion: %d" ), OsVersionInfo.dwMajorVersion, OsVersionInfo.dwMinorVersion ) );
+		OutResults.ErrorCount++;
+		WriteFStringToResults( OutResults.LastSurveyError, TEXT( "GetVersionEx() returned unknown version" ) );
+		WriteFStringToResults( OutResults.LastSurveyErrorDetail, FString::Printf( TEXT( "dwMajorVersion: %d  dwMinorVersion: %d" ), OsVersionInfo.dwMajorVersion, OsVersionInfo.dwMinorVersion ) );
 	}
 
 	if( ErrorCode & FWindowsOSVersionHelper::ERROR_GETVERSIONEX_FAILED )
 	{
 		const uint32 LastError = FPlatformMisc::GetLastError();
 		UE_LOG( LogWindows, Warning, TEXT( "FWindowsPlatformSurvey::GetOSVersionLabel() failed to get Windows version info from GetVersionEx()" ) );
-		Results.ErrorCount++;
-		WriteFStringToResults( Results.LastSurveyError, TEXT( "GetVersionEx() failed" ) );
-		WriteFStringToResults( Results.LastSurveyErrorDetail, FString::Printf( TEXT( "ErrorCode: 0x%0x" ), LastError ) );
+		OutResults.ErrorCount++;
+		WriteFStringToResults( OutResults.LastSurveyError, TEXT( "GetVersionEx() failed" ) );
+		WriteFStringToResults( OutResults.LastSurveyErrorDetail, FString::Printf( TEXT( "ErrorCode: 0x%0x" ), LastError ) );
 	}
 
-	WriteFStringToResults( Results.OSVersion, OSVersionLabel );
-	WriteFStringToResults( Results.OSSubVersion, OSSubVersionLabel );
+	WriteFStringToResults( OutResults.OSVersion, OSVersionLabel );
+	WriteFStringToResults( OutResults.OSSubVersion, OSSubVersionLabel );
 }
 
 bool FWindowsPlatformSurvey::GetLineFollowing(const FString& Token, const TArray<FString>& InLines, FString& OutString, int32 NthHit)
diff --git a/Engine/Source/Runtime/Slate/Private/Widgets/Layout/SWidgetSwitcher.cpp b/Engine/Source/Runtime/Slate/Private/Widgets/Layout/SWidgetSwitcher.cpp
--- a/Engine/Source/Runtime/Slate/Private/Widgets/Layout/SWidgetSwitcher.cpp
+++ b/Engine/Source/Runtime/Slate/Private/Widgets/Layout/SWidgetSwitcher.cpp
@@ -5,15 +5,6 @@
 #include "SWidgetSwitcher.h"
 
 
-/* Local constants
- *****************************************************************************/
-
-const TAttribute<FVector2D> ContentScale = FVector2D::UnitVector;
-
-
-/* SWidgetSwitcher interface
- *****************************************************************************/
-
 SWidgetSwitcher::SWidgetSwitcher()
 	: AllChildren()
 { }
@@ -114,6 +105,8 @@ void SWidgetSwitcher::SetActiveWidgetIndex( int32 Index )
 
 void SWidgetSwitcher::OnArrangeChildren( const FGeometry& AllottedGeometry, FArrangedChildren& ArrangedChildren ) const
 {
+	const TAttribute<FVector2D> ContentScale = FVector2D::UnitVector;
+
 	if (AllChildren.Num() > 0)
 	{
 		ArrangeSingleChild(AllottedGeometry, ArrangedChildren, AllChildren[WidgetIndex.Get()], ContentScale);
diff --git a/Engine/Source/Runtime/Windows/D3D11RHI/Private/D3D11Commands.cpp b/Engine/Source/Runtime/Windows/D3D11RHI/Private/D3D11Commands.cpp
--- a/Engine/Source/Runtime/Windows/D3D11RHI/Private/D3D11Commands.cpp
+++ b/Engine/Source/Runtime/Windows/D3D11RHI/Private/D3D11Commands.cpp
@@ -1496,36 +1496,36 @@ void FD3D11DynamicRHI::RHIClearMRT(bool bClearColor,int32 NumClearColors,const F
 		/** The global D3D device's immediate context */
 		FDeviceStateHelper(TRefCountPtr<FD3D11DeviceContext> InDirect3DDeviceIMContext) : Direct3DDeviceIMContext(InDirect3DDeviceIMContext) {}
 
-		void CaptureDeviceState(FD3D11StateCache& StateCache)
+		void CaptureDeviceState(FD3D11StateCache& StateCacheRef)
 		{
-			StateCache.GetVertexShader(&VSOld);
-			StateCache.GetPixelShader(&PSOld);
-			StateCache.GetShaderResourceViews<SF_Vertex>(0, ResourceCount, &VertResources[0]);
-			StateCache.GetDepthStencilState(&OldDepthStencilState, &StencilRef);
-			StateCache.GetBlendState(&OldBlendState, BlendFactor, &SampleMask);
-			StateCache.GetRasterizerState(&OldRasterizerState);
+			StateCacheRef.GetVertexShader(&VSOld);
+			StateCacheRef.GetPixelShader(&PSOld);
+			StateCacheRef.GetShaderResourceViews<SF_Vertex>(0, ResourceCount, &VertResources[0]);
+			StateCacheRef.GetDepthStencilState(&OldDepthStencilState, &StencilRef);
+			StateCacheRef.GetBlendState(&OldBlendState, BlendFactor, &SampleMask);
+			StateCacheRef.GetRasterizerState(&OldRasterizerState);
 		}
 
-		void ClearCurrentVertexResources(FD3D11StateCache& StateCache)
+		void ClearCurrentVertexResources(FD3D11StateCache& StateCacheRef)
 		{
 			static ID3D11ShaderResourceView* NullResources[ResourceCount] = {};
 			for (int ResourceLoop = 0 ; ResourceLoop < ResourceCount; ResourceLoop++)
 			{
-				StateCache.SetShaderResourceView<SF_Vertex>(NullResources[0],0);
+				StateCacheRef.SetShaderResourceView<SF_Vertex>(NullResources[0],0);
 			}
 		}
 
-		void RestoreDeviceState(FD3D11StateCache& StateCache) 
+		void RestoreDeviceState(FD3D11StateCache& StateCacheRef) 
 		{
 
 			// Restore the old shaders
-			StateCache.SetVertexShader(VSOld);
-			StateCache.SetPixelShader(PSOld);
+			StateCacheRef.SetVertexShader(VSOld);
+			StateCacheRef.SetPixelShader(PSOld);
 			for (int ResourceLoop = 0 ; ResourceLoop < ResourceCount; ResourceLoop++)
-				StateCache.SetShaderResourceView<SF_Vertex>(VertResources[ResourceLoop],ResourceLoop);
-			StateCache.SetDepthStencilState(OldDepthStencilState, StencilRef);
-			StateCache.SetBlendState(OldBlendState, BlendFactor, SampleMask);
-			StateCache.SetRasterizerState(OldRasterizerState);
+				StateCacheRef.SetShaderResourceView<SF_Vertex>(VertResources[ResourceLoop],ResourceLoop);
+			StateCacheRef.SetDepthStencilState(OldDepthStencilState, StencilRef);
+			StateCacheRef.SetBlendState(OldBlendState, BlendFactor, SampleMask);
+			StateCacheRef.SetRasterizerState(OldRasterizerState);
 
 			ReleaseResources();
 		}
diff --git a/Engine/Source/Runtime/Windows/D3D11RHI/Private/D3D11Device.cpp b/Engine/Source/Runtime/Windows/D3D11RHI/Private/D3D11Device.cpp
--- a/Engine/Source/Runtime/Windows/D3D11RHI/Private/D3D11Device.cpp
+++ b/Engine/Source/Runtime/Windows/D3D11RHI/Private/D3D11Device.cpp
@@ -238,14 +238,14 @@ void FD3D11DynamicRHI::RHIGetSupportedResolution( uint32 &Width, uint32 &Height
 	BestMode.Height = 0;
 
 	{
-		HRESULT hr = S_OK;
+		HRESULT HResult = S_OK;
 		TRefCountPtr<IDXGIAdapter> Adapter;
-		hr = DXGIFactory1->EnumAdapters(ChosenAdapter,Adapter.GetInitReference());
-		if( DXGI_ERROR_NOT_FOUND == hr )
+		HResult = DXGIFactory1->EnumAdapters(ChosenAdapter,Adapter.GetInitReference());
+		if( DXGI_ERROR_NOT_FOUND == HResult )
 		{
 			return;
 		}
-		if( FAILED(hr) )
+		if( FAILED(HResult) )
 		{
 			return;
 		}
@@ -259,22 +259,22 @@ void FD3D11DynamicRHI::RHIGetSupportedResolution( uint32 &Width, uint32 &Height
 		for(uint32 o = 0;o < 1; o++)
 		{
 			TRefCountPtr<IDXGIOutput> Output;
-			hr = Adapter->EnumOutputs(o,Output.GetInitReference());
-			if(DXGI_ERROR_NOT_FOUND == hr)
+			HResult = Adapter->EnumOutputs(o,Output.GetInitReference());
+			if(DXGI_ERROR_NOT_FOUND == HResult)
 				break;
-			if(FAILED(hr))
+			if(FAILED(HResult))
 				return;
 
 			// TODO: GetDisplayModeList is a terribly SLOW call.  It can take up to a second per invocation.
 			//  We might want to work around some DXGI badness here.
 			DXGI_FORMAT Format = DXGI_FORMAT_R8G8B8A8_UNORM;
 			uint32 NumModes = 0;
-			hr = Output->GetDisplayModeList(Format,0,&NumModes,NULL);
-			if(hr == DXGI_ERROR_NOT_FOUND)
+			HResult = Output->GetDisplayModeList(Format,0,&NumModes,NULL);
+			if(HResult == DXGI_ERROR_NOT_FOUND)
 			{
 				return;
 			}
-			else if(hr == DXGI_ERROR_NOT_CURRENTLY_AVAILABLE)
+			else if(HResult == DXGI_ERROR_NOT_CURRENTLY_AVAILABLE)
 			{
 				UE_LOG(LogD3D11RHI, Fatal,
 					TEXT("This application cannot be run over a remote desktop configuration")
diff --git a/Engine/Source/Runtime/Windows/D3D11RHI/Private/Windows/WindowsD3D11Device.cpp b/Engine/Source/Runtime/Windows/D3D11RHI/Private/Windows/WindowsD3D11Device.cpp
--- a/Engine/Source/Runtime/Windows/D3D11RHI/Private/Windows/WindowsD3D11Device.cpp
+++ b/Engine/Source/Runtime/Windows/D3D11RHI/Private/Windows/WindowsD3D11Device.cpp
@@ -700,13 +700,13 @@ bool FD3D11DynamicRHI::RHIGetAvailableResolutions(FScreenResolutionArray& Resolu
 		MaxAllowableRefreshRate = 10480;
 	}
 
-	HRESULT hr = S_OK;
+	HRESULT HResult = S_OK;
 	TRefCountPtr<IDXGIAdapter> Adapter;
-	hr = DXGIFactory1->EnumAdapters(ChosenAdapter,Adapter.GetInitReference());
+	HResult = DXGIFactory1->EnumAdapters(ChosenAdapter,Adapter.GetInitReference());
 
-	if( DXGI_ERROR_NOT_FOUND == hr )
+	if( DXGI_ERROR_NOT_FOUND == HResult )
 		return false;
-	if( FAILED(hr) )
+	if( FAILED(HResult) )
 		return false;
 
 	// get the description of the adapter
@@ -720,22 +720,22 @@ bool FD3D11DynamicRHI::RHIGetAvailableResolutions(FScreenResolutionArray& Resolu
 	do 
 	{
 		TRefCountPtr<IDXGIOutput> Output;
-		hr = Adapter->EnumOutputs(CurrentOutput,Output.GetInitReference());
-		if(DXGI_ERROR_NOT_FOUND == hr)
+		HResult = Adapter->EnumOutputs(CurrentOutput,Output.GetInitReference());
+		if(DXGI_ERROR_NOT_FOUND == HResult)
 			break;
-		if(FAILED(hr))
+		if(FAILED(HResult))
 			return false;
 
 		// TODO: GetDisplayModeList is a terribly SLOW call.  It can take up to a second per invocation.
 		//  We might want to work around some DXGI badness here.
 		DXGI_FORMAT Format = DXGI_FORMAT_R8G8B8A8_UNORM;
 		uint32 NumModes = 0;
-		hr = Output->GetDisplayModeList(Format, 0, &NumModes, NULL);
-		if(hr == DXGI_ERROR_NOT_FOUND)
+		HResult = Output->GetDisplayModeList(Format, 0, &NumModes, NULL);
+		if(HResult == DXGI_ERROR_NOT_FOUND)
 		{
 			continue;
 		}
-		else if(hr == DXGI_ERROR_NOT_CURRENTLY_AVAILABLE)
+		else if(HResult == DXGI_ERROR_NOT_CURRENTLY_AVAILABLE)
 		{
 			UE_LOG(LogD3D11RHI, Fatal,
 				TEXT("This application cannot be run over a remote desktop configuration")
diff --git a/Engine/Source/Runtime/Windows/XAudio2/Private/XAudio2Effects.cpp b/Engine/Source/Runtime/Windows/XAudio2/Private/XAudio2Effects.cpp
--- a/Engine/Source/Runtime/Windows/XAudio2/Private/XAudio2Effects.cpp
+++ b/Engine/Source/Runtime/Windows/XAudio2/Private/XAudio2Effects.cpp
@@ -251,12 +251,12 @@ protected:
 	 * 
 	 * @note	Called whenever SetParameters() is called on this XAPO. 
 	 */
-	void OnSetParameters( const void* Parameters, UINT32 ParameterSize )
+	void OnSetParameters( const void* NewParameters, UINT32 ParameterSize )
 	{
 		// The given parameter must be a FAudioRadioEffect struct.
 		check( ParameterSize == sizeof( FAudioRadioEffect ) );
 		
-		FAudioRadioEffect& RadioParams = *( FAudioRadioEffect* )Parameters;
+		FAudioRadioEffect& RadioParams = *( FAudioRadioEffect* )NewParameters;
 	}
 
 private:
diff --git a/Engine/Source/Runtime/Windows/XAudio2/Private/XAudio2Source.cpp b/Engine/Source/Runtime/Windows/XAudio2/Private/XAudio2Source.cpp
--- a/Engine/Source/Runtime/Windows/XAudio2/Private/XAudio2Source.cpp
+++ b/Engine/Source/Runtime/Windows/XAudio2/Private/XAudio2Source.cpp
@@ -390,18 +390,18 @@ bool FXAudio2SoundSource::Init(FWaveInstance* InWaveInstance)
 	if (InWaveInstance->OutputTarget != EAudioOutputTarget::Controller)
 	{
 		// Find matching buffer.
-		FAudioDevice* AudioDevice = nullptr;
+		FAudioDevice* BestAudioDevice = nullptr;
 		if (InWaveInstance->ActiveSound->AudioComponent.IsValid())
 		{
-			AudioDevice = InWaveInstance->ActiveSound->AudioComponent->GetAudioDevice();
+			BestAudioDevice = InWaveInstance->ActiveSound->AudioComponent->GetAudioDevice();
 		}
 		else
 		{
-			AudioDevice = GEngine->GetMainAudioDevice();
+			BestAudioDevice = GEngine->GetMainAudioDevice();
 		}
-		check(AudioDevice);
+		check(BestAudioDevice);
 
-		XAudio2Buffer = FXAudio2SoundBuffer::Init(AudioDevice, InWaveInstance->WaveData, InWaveInstance->StartTime > 0.f);
+		XAudio2Buffer = FXAudio2SoundBuffer::Init(BestAudioDevice, InWaveInstance->WaveData, InWaveInstance->StartTime > 0.f);
 		Buffer = XAudio2Buffer;
 
 		// Buffer failed to be created, or there was an error with the compressed data
