Rocket: disable packaging for projects that need libs which aren't bundled (UE-7114).

- Rocket doesn't bundle Linux libs, making code-based projects (and projects with third-party plugins) fail during compilation.
- Updated messaging to reflect this.
- Also added a SDK check for Linux and a 'getting' started UDN page.
- Updated Linux README for 4.8.

diff --git a/Engine/Build/BatchFiles/Linux/README.md b/Engine/Build/BatchFiles/Linux/README.md
--- a/Engine/Build/BatchFiles/Linux/README.md
+++ b/Engine/Build/BatchFiles/Linux/README.md
@@ -1,7 +1,7 @@
 Build scripts for native Linux build
 ====================================
 
-This document describes how to build Unreal Engine 4.7 natively on a Linux host. 
+This document describes how to build Unreal Engine 4.8 natively on a Linux host. 
 The steps are described here are applicable to the current build, but you may
 want to visit https://wiki.unrealengine.com/Building_On_Linux for the
 latest updates on the process.
@@ -23,10 +23,8 @@ distributions only (Debian itself, (K)Ubuntu and Linux Mint).
 
 Most important dependencies:
 - mono 3.x (2.x may work, but is not recommended), including xbuild and C# compiler (*mcs), and libraries for NET 4.0 framework.
-- clang 3.5.0 (clang 3.3 is also supported, but NOT clang 3.4).
-- Qt(4,5) or GTK development packages to build LinuxNativeDialogs.
-- SDL2 is also needed for building LinuxNativeDialogs module, but the rest
-of the engine is using our own (modified) version of it from Engine/Source/ThirdParty/SDL2.
+- clang 3.5.0 (clang 3.6 and 3.3 are also supported, but NOT clang 3.4).
+- For non-Ubuntu OSes: Qt(4,5) or GTK development packages to build LinuxNativeDialogs.
 
 You will also need at least 20 GB of free disk space and a relatively powerful
 machine.
@@ -51,7 +49,7 @@ How to set up the sources for building, step by step:
 
 1. Clone EpicGames/UnrealEngine repository
 
-    ``git clone https://github.com/EpicGames/UnrealEngine -b 4.7``
+    ``git clone https://github.com/EpicGames/UnrealEngine -b 4.8``
     
 2. Run Setup.sh once.
 
@@ -80,16 +78,21 @@ registered as a post-merge hook by Setup.sh, so third party libraries will be up
 Building and running
 --------------------
 
-GenerateProjectFiles.sh also produces both makefile and CMakeLists.txt which you can use to import the
-project in your favorite IDE. KDevelop 4.6+ is known to handle the project well 
-(it takes about 3-4 GB of resident RAM to load the project though).
+GenerateProjectFiles.sh produces a number of "project" files, including Makefile, CMakeLists.txt which you can use to import the
+project in your favorite IDE and qmake project file. Both QtCreator and KDevelop 4.6+ are known to handle the project well 
+(although the latter takes about 3-4 GB of resident RAM to load the project).
 
 The targets match the name of the resulting binary, e.g. UE4Editor-Linux-Debug or UE4Game. You can build them
 by just typing make <target> in the engine's root folder.
 
-Specifically, to be able to run the editor, build the following targets:
+4.8 and later versions have simplified building the editor by providing makefile targets "StandardSet" and "DebugSet", the former
+being the default. You can now just type
 
-    make ShaderCompileWorker UnrealLightmass UnrealPak UE4Editor
+    make
+
+to build the editor. Alternatively, build the following targets:
+
+    make CrashReportClient ShaderCompileWorker UnrealLightmass UnrealPak UE4Editor
 
 If you intend to develop the editor, you can build a debug configuration of it:
 
@@ -132,7 +135,8 @@ Depending on the project, the editor may need rather large number of file handle
 you may need to adjust your limits.
 
 It is advised that you install the editor on a case insensitive filesystem if you
-intend to open projects from other systems (OS X and Windows), or Marketplace.
+intend to open projects from other systems (OS X and Windows), or Marketplace. JFS formatted
+with -O option seems to be the best solution.
 
 The time it takes to build the editor in development configuration can be large,
 debug configuration takes about 2/3 of this time. The build process can also take 
diff --git a/Engine/Documentation/Source/Platforms/Linux/GettingStarted/LinuxCrossGettingStarted.INT.udn b/Engine/Documentation/Source/Platforms/Linux/GettingStarted/LinuxCrossGettingStarted.INT.udn
new file mode 100644
--- /dev/null
+++ b/Engine/Documentation/Source/Platforms/Linux/GettingStarted/LinuxCrossGettingStarted.INT.udn
@@ -0,0 +1,95 @@
+﻿Availability:Docs
+Title: Getting Started: Developing Linux Projects
+Crumbs:%ROOT%, Platforms, Platforms/Linux
+Description:Setting up for development for the Linux platform.
+
+[TOC(start:2 end:3)]
+
+## Disclaimer
+
+%Globals:BinarySourceSelectVar%
+
+[OBJECT:ToggleButtonContent]
+	[PARAMLITERAL:category]
+	BinarySourceSelect
+	[/PARAMLITERAL]
+	[PARAMLITERAL:id]
+	binaryinstall
+	[/PARAMLITERAL]
+	[PARAMLITERAL:active]
+	active_button_content
+	[/PARAMLITERAL]
+	[PARAM:content]
+	%binary%
+	[/PARAM]
+[/OBJECT]
+[OBJECT:ToggleButtonContent]
+	[PARAMLITERAL:category]
+	BinarySourceSelect
+	[/PARAMLITERAL]
+	[PARAMLITERAL:id]
+	sourceinstall
+	[/PARAMLITERAL]
+	[PARAMLITERAL:active]
+	[/PARAMLITERAL]
+	[PARAM:content]
+	%source%
+	[/PARAM]
+[/OBJECT]
+
+
+[VAR:source]
+## Source and Github Users
+
+Linux as a target platform is supported on Windows host platform and on Linux itself (but not yet on Mac). This guide is geared towards Windows users trying to cross-compile for Linux; if you run Unreal Editor on Linux itself, no separate step is needed at the moment
+(note that on Linux, we currently use the compiler available on the system, although in the future we may bundle it to provide a stable build environment which will allow you to prepare the binaries in a distro-agnostic way).
+
+On Windows, we use the [cross-compiler toolchain](http://cdn.unrealengine.com/qfe/v4_clang-3.5.0_ld-2.24_glibc-2.12.2.zip) to cross compile a Linux binary. 
+
+### Installing from Source: 
+
+As this guide assumes that you are running Windows, you will need to install the toolchain first:
+
+* [Clang-3.5.0 based tool chain](http://cdn.unrealengine.com/qfe/v4_clang-3.5.0_ld-2.24_glibc-2.12.2.zip)
+
+### Toolchain set up
+
+Unzip the above file to a path that does not contain spaces or non-ASCII characters (e.g. C:\CrossToolchain350 is a good example) and set LINUX_ROOT environment variable so it contains the absolute path to the "toolchain" directory of the toolchain (e.g. C:\CrossToolchain350\toolchain\).
+More information: https://wiki.unrealengine.com/Compiling_For_Linux
+
+### Unreal Engine 4 Setup. 
+
+* Fork and clone the Unreal Engine 4 repository \[[UnrealEngine](https://github.com/EpicGames/UnrealEngine)\]
+
+* Within a shell run `Setup.bat -include=Linux`
+
+Follow the instuctions on [https://github.com/EpicGames/UnrealEngine](https://github.com/EpicGames/UnrealEngine) to build the editor.  Set your solution configuration to **Development** and your solution platform to **Linux** to build Linux targets.
+
+[/VAR]
+
+
+
+
+
+
+[VAR:binary]
+
+## Binary and Launcher Users
+
+Linux as a target platform is supported on Windows host platform and on Linux itself (but not yet on Mac). This guide is geared towards Windows users trying to cross-compile for Linux; if you run Unreal Editor on Linux itself, no separate step is needed at the moment
+(note that on Linux, we currently use the compiler available on the system, although in the future we may bundle it to provide a stable build environment which will allow you to prepare the binaries in a distro-agnostic way).
+
+On Windows, we use the [cross-compiler toolchain](http://cdn.unrealengine.com/qfe/v4_clang-3.5.0_ld-2.24_glibc-2.12.2.zip) to cross compile a Linux binary. 
+
+### Installing 
+
+As this guide assumes that you are running Windows, you will need to install the toolchain first:
+
+* [Clang-3.5.0 based tool chain](http://cdn.unrealengine.com/qfe/v4_clang-3.5.0_ld-2.24_glibc-2.12.2.zip)
+
+### Toolchain set up
+
+Unzip the above file to a path that does not contain spaces or non-ASCII characters (e.g. C:\CrossToolchain350 is a good example) and set LINUX_ROOT environment variable so it contains the absolute path to the "toolchain" directory of the toolchain (e.g. C:\CrossToolchain350\toolchain\).
+More information: https://wiki.unrealengine.com/Compiling_For_Linux
+
+[/VAR]
diff --git a/Engine/Documentation/Source/Platforms/Linux/Linux.INT.udn b/Engine/Documentation/Source/Platforms/Linux/Linux.INT.udn
new file mode 100644
--- /dev/null
+++ b/Engine/Documentation/Source/Platforms/Linux/Linux.INT.udn
@@ -0,0 +1,46 @@
+Availability: Docs
+Title: Linux Game Development
+Crumbs:%ROOT%, Platforms
+Description:Developing games for the Linux platform.
+Platform:Linux
+Navigation:platform
+
+
+[OBJECT:TopicButtonList]
+	[PARAM:icon]
+		![](%ROOT%/start_icon.png)(convert:false)
+	[/PARAM]
+	[PARAM:title]
+		Linux Essentials
+	[/PARAM]
+	[PARAM:description]
+	[/PARAM]
+	[PARAM:links]
+	* [](Platforms/Linux/GettingStarted "%Platforms/Linux/GettingStarted:description%")
+	[/PARAM]
+[/OBJECT]
+
+	
+
+		
+
+	      
+
+	       
+
+	   
+
+	   
+
+	  
+
+		
+
+		
+
+		
+
+		
+
+		
+
diff --git a/Engine/Source/Developer/DesktopPlatform/Private/PlatformInfo.cpp b/Engine/Source/Developer/DesktopPlatform/Private/PlatformInfo.cpp
--- a/Engine/Source/Developer/DesktopPlatform/Private/PlatformInfo.cpp
+++ b/Engine/Source/Developer/DesktopPlatform/Private/PlatformInfo.cpp
@@ -88,10 +88,10 @@ static const FPlatformInfo PlatformInfoArray[] = {
 	BuildPlatformInfo(TEXT("MacClient"),				TEXT("MacClient"),			LOCTEXT("MacClient", "Mac (Client-only)"),					EPlatformType::Client,		EPlatformFlags::None,			FPlatformIconPaths(TEXT("Launcher/MacTarget/Platform_Mac_24x"), TEXT("Launcher/MacTarget/Platform_Mac_128x")),										TEXT(""),												TEXT(""),			EPlatformSDKStatus::Unknown,	TEXT("/Engine/Tutorial/Installation/InstallingXCodeTutorial.InstallingXCodeTutorial"),					IsAvailableOnMac,							true,				IsAvailableOnMac),
 	BuildPlatformInfo(TEXT("MacServer"),				TEXT("MacServer"),			LOCTEXT("MacServer", "Mac (Dedicated Server)"),				EPlatformType::Server,		EPlatformFlags::None,			FPlatformIconPaths(TEXT("Launcher/MacTarget/Platform_Mac_24x"), TEXT("Launcher/MacTarget/Platform_Mac_128x")),										TEXT(""),												TEXT(""),			EPlatformSDKStatus::Unknown,	TEXT("/Engine/Tutorial/Installation/InstallingXCodeTutorial.InstallingXCodeTutorial"),					IsAvailableOnMac,							true,				IsAvailableOnMac),
 
-	BuildPlatformInfo(TEXT("LinuxNoEditor"),			TEXT("LinuxNoEditor"),		LOCTEXT("LinuxNoEditor", "Linux"),							EPlatformType::Game,		EPlatformFlags::None,			FPlatformIconPaths(TEXT("Launcher/LinuxTarget/Platform_Linux_24x"), TEXT("Launcher/LinuxTarget/Platform_Linux_128x")),								TEXT("-targetplatform=Linux"),							TEXT("Linux_x64"),	EPlatformSDKStatus::Unknown,	TEXT(""),																								IsAvailableOnLinux || IsAvailableOnWindows,	 true,				IsAvailableOnLinux),	
-	BuildPlatformInfo(TEXT("Linux"),					TEXT("Linux"),				LOCTEXT("LinuxEditor", "Linux (Editor)"),					EPlatformType::Editor,		EPlatformFlags::None,			FPlatformIconPaths(TEXT("Launcher/LinuxTarget/Platform_Linux_24x"), TEXT("Launcher/LinuxTarget/Platform_Linux_128x")),								TEXT(""),												TEXT("Linux_x64"),	EPlatformSDKStatus::Unknown,	TEXT(""),																								IsAvailableOnLinux,							 true,				IsAvailableOnLinux),
-	BuildPlatformInfo(TEXT("LinuxClient"),				TEXT("LinuxClient"),		LOCTEXT("LinuxClient", "Linux (Client-only)"),				EPlatformType::Client,		EPlatformFlags::None,			FPlatformIconPaths(TEXT("Launcher/LinuxTarget/Platform_Linux_24x"), TEXT("Launcher/LinuxTarget/Platform_Linux_128x")),								TEXT(""),												TEXT("Linux_x64"),	EPlatformSDKStatus::Unknown,	TEXT(""),																								IsAvailableOnLinux || IsAvailableOnWindows,	 true,				IsAvailableOnLinux),
-	BuildPlatformInfo(TEXT("LinuxServer"),				TEXT("LinuxServer"),		LOCTEXT("LinuxServer", "Linux (Dedicated Server)"),			EPlatformType::Server,		EPlatformFlags::None,			FPlatformIconPaths(TEXT("Launcher/LinuxTarget/Platform_Linux_24x"), TEXT("Launcher/LinuxTarget/Platform_Linux_128x")),								TEXT(""),												TEXT("Linux_x64"),	EPlatformSDKStatus::Unknown,	TEXT(""),																								IsAvailableOnLinux || IsAvailableOnWindows,	 true,				IsAvailableOnLinux),
+	BuildPlatformInfo(TEXT("LinuxNoEditor"),			TEXT("LinuxNoEditor"),		LOCTEXT("LinuxNoEditor", "Linux"),							EPlatformType::Game,		EPlatformFlags::None,			FPlatformIconPaths(TEXT("Launcher/LinuxTarget/Platform_Linux_24x"), TEXT("Launcher/LinuxTarget/Platform_Linux_128x")),								TEXT("-targetplatform=Linux"),							TEXT("Linux_x64"),	EPlatformSDKStatus::Unknown,	TEXT("/Platforms/Linux/GettingStarted"),																IsAvailableOnLinux || IsAvailableOnWindows,	 true,				IsAvailableOnLinux),	
+	BuildPlatformInfo(TEXT("Linux"),					TEXT("Linux"),				LOCTEXT("LinuxEditor", "Linux (Editor)"),					EPlatformType::Editor,		EPlatformFlags::None,			FPlatformIconPaths(TEXT("Launcher/LinuxTarget/Platform_Linux_24x"), TEXT("Launcher/LinuxTarget/Platform_Linux_128x")),								TEXT(""),												TEXT("Linux_x64"),	EPlatformSDKStatus::Unknown,	TEXT("/Platforms/Linux/GettingStarted"),																IsAvailableOnLinux,							 true,				IsAvailableOnLinux),
+	BuildPlatformInfo(TEXT("LinuxClient"),				TEXT("LinuxClient"),		LOCTEXT("LinuxClient", "Linux (Client-only)"),				EPlatformType::Client,		EPlatformFlags::None,			FPlatformIconPaths(TEXT("Launcher/LinuxTarget/Platform_Linux_24x"), TEXT("Launcher/LinuxTarget/Platform_Linux_128x")),								TEXT(""),												TEXT("Linux_x64"),	EPlatformSDKStatus::Unknown,	TEXT("/Platforms/Linux/GettingStarted"),																IsAvailableOnLinux || IsAvailableOnWindows,	 true,				IsAvailableOnLinux),
+	BuildPlatformInfo(TEXT("LinuxServer"),				TEXT("LinuxServer"),		LOCTEXT("LinuxServer", "Linux (Dedicated Server)"),			EPlatformType::Server,		EPlatformFlags::None,			FPlatformIconPaths(TEXT("Launcher/LinuxTarget/Platform_Linux_24x"), TEXT("Launcher/LinuxTarget/Platform_Linux_128x")),								TEXT(""),												TEXT("Linux_x64"),	EPlatformSDKStatus::Unknown,	TEXT("/Platforms/Linux/GettingStarted"),																IsAvailableOnLinux || IsAvailableOnWindows,	 true,				IsAvailableOnLinux),
 
 	BuildPlatformInfo(TEXT("IOS"),						TEXT("IOS"),				LOCTEXT("IOS", "iOS"),										EPlatformType::Game,		EPlatformFlags::None,			FPlatformIconPaths(TEXT("Launcher/iOSTarget/Platform_iOS_24x"), TEXT("Launcher/iOSTarget/Platform_iOS_128x")),										TEXT(""),												TEXT(""),			EPlatformSDKStatus::Unknown,	IOSTutorial,																							IsAvailableOnWindows || IsAvailableOnMac,	true,				false),
 
diff --git a/Engine/Source/Developer/Linux/LinuxNoEditorTargetPlatform/LinuxNoEditorTargetPlatform.Build.cs b/Engine/Source/Developer/Linux/LinuxNoEditorTargetPlatform/LinuxNoEditorTargetPlatform.Build.cs
--- a/Engine/Source/Developer/Linux/LinuxNoEditorTargetPlatform/LinuxNoEditorTargetPlatform.Build.cs
+++ b/Engine/Source/Developer/Linux/LinuxNoEditorTargetPlatform/LinuxNoEditorTargetPlatform.Build.cs
@@ -13,6 +13,7 @@ public class LinuxNoEditorTargetPlatform : ModuleRules
 				"Core",
 				"TargetPlatform",
 				"DesktopPlatform",
+				"Projects"
 			}
         );
 
diff --git a/Engine/Source/Developer/Linux/LinuxServerTargetPlatform/LinuxServerTargetPlatform.Build.cs b/Engine/Source/Developer/Linux/LinuxServerTargetPlatform/LinuxServerTargetPlatform.Build.cs
--- a/Engine/Source/Developer/Linux/LinuxServerTargetPlatform/LinuxServerTargetPlatform.Build.cs
+++ b/Engine/Source/Developer/Linux/LinuxServerTargetPlatform/LinuxServerTargetPlatform.Build.cs
@@ -13,6 +13,7 @@ public class LinuxServerTargetPlatform : ModuleRules
 				"Core",
 				"TargetPlatform",
 				"DesktopPlatform",
+				"Projects"
 			}
         );
 
diff --git a/Engine/Source/Developer/Linux/LinuxTargetPlatform/LinuxTargetPlatform.Build.cs b/Engine/Source/Developer/Linux/LinuxTargetPlatform/LinuxTargetPlatform.Build.cs
--- a/Engine/Source/Developer/Linux/LinuxTargetPlatform/LinuxTargetPlatform.Build.cs
+++ b/Engine/Source/Developer/Linux/LinuxTargetPlatform/LinuxTargetPlatform.Build.cs
@@ -14,6 +14,7 @@ public class LinuxTargetPlatform : ModuleRules
 				"CoreUObject",
 				"TargetPlatform",
 				"DesktopPlatform",
+				"Projects"
 			}
         );
 
diff --git a/Engine/Source/Developer/Linux/LinuxTargetPlatform/Private/LinuxTargetPlatform.h b/Engine/Source/Developer/Linux/LinuxTargetPlatform/Private/LinuxTargetPlatform.h
--- a/Engine/Source/Developer/Linux/LinuxTargetPlatform/Private/LinuxTargetPlatform.h
+++ b/Engine/Source/Developer/Linux/LinuxTargetPlatform/Private/LinuxTargetPlatform.h
@@ -9,6 +9,7 @@
 #if WITH_ENGINE
 #include "StaticMeshResources.h"
 #endif // WITH_ENGINE
+#include "IProjectManager.h"
 
 #define LOCTEXT_NAMESPACE "TLinuxTargetPlatform"
 
@@ -150,6 +151,57 @@ public:
 		return TTargetPlatformBase<FLinuxPlatformProperties<HAS_EDITOR_DATA, IS_DEDICATED_SERVER, IS_CLIENT_ONLY>>::SupportsFeature(Feature);
 	}
 
+	virtual bool IsSdkInstalled(bool bProjectHasCode, FString& OutDocumentationPath) const override
+	{
+		if (!PLATFORM_LINUX)
+		{
+			// check for LINUX_ROOT when targeting Linux from Win/Mac
+			TCHAR ToolchainRoot[32768] = { 0 };
+			FPlatformMisc::GetEnvironmentVariable(TEXT("LINUX_ROOT"), ToolchainRoot, ARRAY_COUNT(ToolchainRoot));
+
+			FString ToolchainCompiler = ToolchainRoot;
+			if (PLATFORM_WINDOWS)
+			{
+				ToolchainCompiler += "/bin/clang++.exe";
+			}
+			else if (PLATFORM_MAC)
+			{
+				ToolchainCompiler += "/bin/clang++";
+			}
+			else
+			{
+				checkf(false, TEXT("Unable to target Linux on an unknown platform."));
+				return false;
+			}
+
+			return FPaths::FileExists(ToolchainCompiler);
+		}
+
+		return true;
+	}
+
+	virtual int32 CheckRequirements(const FString& ProjectPath, bool bProjectHasCode, FString& OutDocumentationPath) const override
+	{
+		int32 ReadyToBuild = TSuper::CheckRequirements(ProjectPath, bProjectHasCode, OutDocumentationPath);
+
+		// do not support code/plugins in Rocket as the required libs aren't bundled (on Windows/Mac)
+		if (!PLATFORM_LINUX && FRocketSupport::IsRocket())
+		{
+			if (bProjectHasCode)
+			{
+				ReadyToBuild |= ETargetPlatformReadyStatus::CodeUnsupported;
+			}
+
+			if (IProjectManager::Get().IsNonDefaultPluginEnabled())
+			{
+				ReadyToBuild |= ETargetPlatformReadyStatus::PluginsUnsupported;
+			}
+		}
+
+		return ReadyToBuild;
+	}
+
+
 #if WITH_ENGINE
 	virtual void GetAllPossibleShaderFormats( TArray<FName>& OutFormats ) const override
 	{
diff --git a/Engine/Source/Editor/MainFrame/Private/Frame/MainFrameActions.cpp b/Engine/Source/Editor/MainFrame/Private/Frame/MainFrameActions.cpp
--- a/Engine/Source/Editor/MainFrame/Private/Frame/MainFrameActions.cpp
+++ b/Engine/Source/Editor/MainFrame/Private/Frame/MainFrameActions.cpp
@@ -435,18 +435,16 @@ void FMainFrameActionCallbacks::PackageProject( const FName InPlatformInfoName )
 			// report to main frame
 			bool UnrecoverableError = false;
 
-#if PLATFORM_WINDOWS
 			if ((Result & ETargetPlatformReadyStatus::CodeUnsupported) != 0)
 			{
-				FMessageDialog::Open(EAppMsgType::Ok, LOCTEXT("IOSNotSupported", "Sorry, launching on a device is currently not supported for code-based iOS projects. This feature will be available in a future release.") );
+				FMessageDialog::Open(EAppMsgType::Ok, LOCTEXT("NotSupported", "Sorry, packaging a code-based project for the selected platform is currently not supported. This feature may be available in a future release."));
 				UnrecoverableError = true;
 			}
 			else if ((Result & ETargetPlatformReadyStatus::PluginsUnsupported) != 0)
 			{
-				FMessageDialog::Open(EAppMsgType::Ok, LOCTEXT("IOSNotSupported", "Sorry, launching on a device is currently not supported for content based projects with third-party plugins. This feature will be available in a future release.") );
+				FMessageDialog::Open(EAppMsgType::Ok, LOCTEXT("NotSupported", "Sorry, packaging a project with third-party plugins is currently not supported for the selected platform. This feature may be available in a future release."));
 				UnrecoverableError = true;
 			}
-#endif
 
 			if (UnrecoverableError)
 			{
diff --git a/Engine/Source/Editor/UnrealEd/Private/Kismet2/DebuggerCommands.cpp b/Engine/Source/Editor/UnrealEd/Private/Kismet2/DebuggerCommands.cpp
--- a/Engine/Source/Editor/UnrealEd/Private/Kismet2/DebuggerCommands.cpp
+++ b/Engine/Source/Editor/UnrealEd/Private/Kismet2/DebuggerCommands.cpp
@@ -1565,17 +1565,15 @@ bool FInternalPlayWorldCommandCallbacks::IsReadyToLaunchOnDevice(FString DeviceI
 		// report to main frame
 		switch (Result)
 		{
-#if PLATFORM_WINDOWS
 		case ETargetPlatformReadyStatus::CodeUnsupported:
 			// show the message
-			FMessageDialog::Open(EAppMsgType::Ok, LOCTEXT("IOSNotSupported", "Sorry, launching on a device is currently not supported for code-based iOS projects. This feature will be available in a future release.") );
+			FMessageDialog::Open(EAppMsgType::Ok, LOCTEXT("NotSupported", "Sorry, packaging a code-based project for the selected platform is currently not supported. This feature may be available in a future release."));
 			return false;
 
 		case ETargetPlatformReadyStatus::PluginsUnsupported:
 			// show the message
-			FMessageDialog::Open(EAppMsgType::Ok, LOCTEXT("IOSNotSupported", "Sorry, launching on a device is currently not supported for content based projects with third-party plugins. This feature will be available in a future release.") );
+			FMessageDialog::Open(EAppMsgType::Ok, LOCTEXT("NotSupported", "Sorry, packaging a project with third-party plugins is currently not supported for the selected platform. This feature may be available in a future release."));
 			return false;
-#endif
 		default:
 			break;
 		}
