Disable UBT makefiles when compiling projects from the editor. It's more important to be robust than fast, since we're automating something for the user.

diff --git a/Engine/Source/Developer/DesktopPlatform/Private/DesktopPlatformBase.cpp b/Engine/Source/Developer/DesktopPlatform/Private/DesktopPlatformBase.cpp
--- a/Engine/Source/Developer/DesktopPlatform/Private/DesktopPlatformBase.cpp
+++ b/Engine/Source/Developer/DesktopPlatform/Private/DesktopPlatformBase.cpp
@@ -441,7 +441,7 @@ bool FDesktopPlatformBase::CompileGameProject(const FString& RootDir, const FStr
 	}
 
 	// Append any other options
-	Arguments += " -editorrecompile -progress";
+	Arguments += " -editorrecompile -progress -noubtmakefiles";
 
 	// Run UBT
 	return RunUnrealBuildTool(LOCTEXT("CompilingProject", "Compiling project..."), RootDir, Arguments, Warn);
