Fixed up mac not converting paths to relative paths when saving filenames.

diff --git a/Engine/Source/Developer/DesktopPlatform/Private/Mac/DesktopPlatformMac.cpp b/Engine/Source/Developer/DesktopPlatform/Private/Mac/DesktopPlatformMac.cpp
--- a/Engine/Source/Developer/DesktopPlatform/Private/Mac/DesktopPlatformMac.cpp
+++ b/Engine/Source/Developer/DesktopPlatform/Private/Mac/DesktopPlatformMac.cpp
@@ -515,10 +515,20 @@ bool FDesktopPlatformMac::FileDialogShared(bool bSave, const void* ParentWindowH
 					OutFilterIndex = [AccessoryView SelectedExtension];
 				}
 
-				// Make sure all filenames gathered have their paths normalized
-				for (auto FilenameIt = OutFilenames.CreateIterator(); FilenameIt; ++FilenameIt)
+				// Make sure all filenames gathered have their paths normalized and proper extensions added
+				for (auto OutFilenameIt = OutFilenames.CreateIterator(); OutFilenameIt; ++OutFilenameIt)
 				{
-					FPaths::NormalizeFilename(*FilenameIt);
+					FString& OutFilename = *OutFilenameIt;
+
+					OutFilename = IFileManager::Get().ConvertToRelativePath(*OutFilename);
+
+					if (FPaths::GetExtension(OutFilename).IsEmpty() && !Extension.IsEmpty())
+					{
+						// filename does not have an extension. Add an extension based on the filter that the user chose in the dialog
+						OutFilename += Extension;
+					}
+
+					FPaths::NormalizeFilename(OutFilename);
 				}
 
 				bOkPressed = true;
