[INTEGRATE] merge ue4<-ue4-launcherDev cl 2582631

#change Rollback of accidental merge 2581810. Discussed alternate solution with Ben.Marsh

diff --git a/Engine/Source/Developer/DesktopPlatform/Private/DesktopPlatformBase.cpp b/Engine/Source/Developer/DesktopPlatform/Private/DesktopPlatformBase.cpp
--- a/Engine/Source/Developer/DesktopPlatform/Private/DesktopPlatformBase.cpp
+++ b/Engine/Source/Developer/DesktopPlatform/Private/DesktopPlatformBase.cpp
@@ -690,24 +690,13 @@ bool FDesktopPlatformBase::GetSolutionPath(FString& OutSolutionPath)
 	return false;
 }
 
-FString FDesktopPlatformBase::GetDefaultProjectCreationPath() const
+FString FDesktopPlatformBase::GetDefaultProjectCreationPath()
 {
 	// My Documents
 	const FString DefaultProjectSubFolder = TEXT("Unreal Projects");
 	return FString(FPlatformProcess::UserDir()) + DefaultProjectSubFolder;
 }
 
-FString FDesktopPlatformBase::GetVersionSelectorPath() const
-{
-	// Get the path to VersionSelector.exe. If we're running from UnrealVersionSelector itself, try to stick with the current configuration.
-	FString DefaultVersionSelectorName = FPlatformProcess::ExecutableName(false);
-	if (!DefaultVersionSelectorName.StartsWith(TEXT("UnrealVersionSelector")))
-	{
-		DefaultVersionSelectorName = TEXT("UnrealVersionSelector.exe");
-	}
-	return FString(FPlatformProcess::BaseDir()) / DefaultVersionSelectorName;
-}
-
 void FDesktopPlatformBase::ReadLauncherInstallationList()
 {
 	FString InstalledListFile = FString(FPlatformProcess::ApplicationSettingsDir()) / TEXT("UnrealEngineLauncher/LauncherInstalled.dat");
diff --git a/Engine/Source/Developer/DesktopPlatform/Private/DesktopPlatformBase.h b/Engine/Source/Developer/DesktopPlatform/Private/DesktopPlatformBase.h
--- a/Engine/Source/Developer/DesktopPlatform/Private/DesktopPlatformBase.h
+++ b/Engine/Source/Developer/DesktopPlatform/Private/DesktopPlatformBase.h
@@ -48,8 +48,7 @@ public:
 	virtual bool GetSolutionPath(FString& OutSolutionPath) override;
 
 	virtual bool EnumerateProjectsKnownByEngine(const FString &Identifier, bool bIncludeNativeProjects, TArray<FString> &OutProjectFileNames) override;
-	virtual FString GetDefaultProjectCreationPath() const override;
-	virtual FString GetVersionSelectorPath() const override;
+	virtual FString GetDefaultProjectCreationPath() override;
 
 private:
 	FString CurrentEngineIdentifier;
diff --git a/Engine/Source/Developer/DesktopPlatform/Private/Windows/DesktopPlatformWindows.cpp b/Engine/Source/Developer/DesktopPlatform/Private/Windows/DesktopPlatformWindows.cpp
--- a/Engine/Source/Developer/DesktopPlatform/Private/Windows/DesktopPlatformWindows.cpp
+++ b/Engine/Source/Developer/DesktopPlatform/Private/Windows/DesktopPlatformWindows.cpp
@@ -594,11 +594,17 @@ FString FDesktopPlatformWindows::GetUserTempPath()
 	return FString(FPlatformProcess::UserTempDir());
 }
 
-FString FDesktopPlatformWindows::GetVersionSelectorPath() const
+void FDesktopPlatformWindows::GetRequiredRegistrySettings(TIndirectArray<FRegistryRootedKey> &RootedKeys)
 {
-	FString ExecutableFileName = FDesktopPlatformBase::GetVersionSelectorPath();
+	// Get the path to VersionSelector.exe. If we're running from UnrealVersionSelector itself, try to stick with the current configuration.
+	FString DefaultVersionSelectorName = FPlatformProcess::ExecutableName(false);
+	if (!DefaultVersionSelectorName.StartsWith(TEXT("UnrealVersionSelector")))
+	{
+		DefaultVersionSelectorName = TEXT("UnrealVersionSelector.exe");
+	}
+	FString ExecutableFileName = FString(FPlatformProcess::BaseDir()) / DefaultVersionSelectorName;
 
-	// Defer to UnrealVersionSelector in a launcher installation if it's got the same version number or greater.
+	// Defer to UnrealVersionSelector.exe in a launcher installation if it's got the same version number or greater.
 	FString InstallDir;
 	if (FWindowsPlatformMisc::QueryRegKey(HKEY_LOCAL_MACHINE, TEXT("SOFTWARE\\EpicGames\\Unreal Engine"), TEXT("INSTALLDIR"), InstallDir) && (InstallDir.Len() > 0))
 	{
@@ -606,19 +612,13 @@ FString FDesktopPlatformWindows::GetVersionSelectorPath() const
 		FPaths::NormalizeDirectoryName(NormalizedInstallDir);
 
 		FString InstalledExecutableFileName = NormalizedInstallDir / FString("Launcher/Engine/Binaries/Win64/UnrealVersionSelector.exe");
-		if(GetShellIntegrationVersion(InstalledExecutableFileName) >= GetShellIntegrationVersion(ExecutableFileName))
+		if(GetShellIntegrationVersion(InstalledExecutableFileName) == GetShellIntegrationVersion(ExecutableFileName))
 		{
 			ExecutableFileName = InstalledExecutableFileName;
 		}
 	}
 
-	return ExecutableFileName;
-}
-
-void FDesktopPlatformWindows::GetRequiredRegistrySettings(TIndirectArray<FRegistryRootedKey> &RootedKeys)
-{
 	// Get the path to the executable
-	FString ExecutableFileName = GetVersionSelectorPath();
 	FPaths::MakePlatformFilename(ExecutableFileName);
 	FString QuotedExecutableFileName = FString(TEXT("\"")) + ExecutableFileName + FString(TEXT("\""));
 
@@ -688,7 +688,7 @@ void FDesktopPlatformWindows::GetRequiredRegistrySettings(TIndirectArray<FRegist
 	}
 }
 
-int32 FDesktopPlatformWindows::GetShellIntegrationVersion(const FString &FileName) const
+int32 FDesktopPlatformWindows::GetShellIntegrationVersion(const FString &FileName)
 {
 	::DWORD VersionInfoSize = GetFileVersionInfoSize(*FileName, NULL);
 	if (VersionInfoSize != 0)
diff --git a/Engine/Source/Developer/DesktopPlatform/Private/Windows/DesktopPlatformWindows.h b/Engine/Source/Developer/DesktopPlatform/Private/Windows/DesktopPlatformWindows.h
--- a/Engine/Source/Developer/DesktopPlatform/Private/Windows/DesktopPlatformWindows.h
+++ b/Engine/Source/Developer/DesktopPlatform/Private/Windows/DesktopPlatformWindows.h
@@ -33,14 +33,13 @@ public:
 	virtual FFeedbackContext* GetNativeFeedbackContext() override;
 
 	virtual FString GetUserTempPath() override;
-	virtual FString GetVersionSelectorPath() const override;
 
 private:
 	bool GetLauncherPath(FString& OutLauncherPath);
 	bool GetLauncherInstallerPath(FString& OutInstallerPath);
 	bool FileDialogShared(bool bSave, const void* ParentWindowHandle, const FString& DialogTitle, const FString& DefaultPath, const FString& DefaultFile, const FString& FileTypes, uint32 Flags, TArray<FString>& OutFilenames, int32& OutFilterIndex);
 	void GetRequiredRegistrySettings(TIndirectArray<FRegistryRootedKey> &RootedKeys);
-	int32 GetShellIntegrationVersion(const FString &FileName) const;
+	int32 GetShellIntegrationVersion(const FString &FileName);
 };
 
 typedef FDesktopPlatformWindows FDesktopPlatform;
diff --git a/Engine/Source/Developer/DesktopPlatform/Public/IDesktopPlatform.h b/Engine/Source/Developer/DesktopPlatform/Public/IDesktopPlatform.h
--- a/Engine/Source/Developer/DesktopPlatform/Public/IDesktopPlatform.h
+++ b/Engine/Source/Developer/DesktopPlatform/Public/IDesktopPlatform.h
@@ -435,12 +435,5 @@ public:
 	* 
 	* @return Path to the folder
 	*/
-	virtual FString GetDefaultProjectCreationPath() const = 0;
-
-	/**
-	* Gets the path to the unreal version selector
-	*
-	* @return Path to the folder
-	*/
-	virtual FString GetVersionSelectorPath() const = 0;
+	virtual FString GetDefaultProjectCreationPath() = 0;
 };
\ No newline at end of file
