[INTEGRATE] merge ue4<-ue4-launcherDev cl 2582701

#change Aalternate solution discussed with [REDACTED_PERSON]

diff --git a/Engine/Source/Developer/DesktopPlatform/Private/Windows/DesktopPlatformWindows.cpp b/Engine/Source/Developer/DesktopPlatform/Private/Windows/DesktopPlatformWindows.cpp
--- a/Engine/Source/Developer/DesktopPlatform/Private/Windows/DesktopPlatformWindows.cpp
+++ b/Engine/Source/Developer/DesktopPlatform/Private/Windows/DesktopPlatformWindows.cpp
@@ -602,7 +602,7 @@ void FDesktopPlatformWindows::GetRequiredRegistrySettings(TIndirectArray<FRegist
 	{
 		DefaultVersionSelectorName = TEXT("UnrealVersionSelector.exe");
 	}
-	FString ExecutableFileName = FString(FPlatformProcess::BaseDir()) / DefaultVersionSelectorName;
+	FString ExecutableFileName = FPaths::ConvertRelativePathToFull(FPaths::EngineDir()) / TEXT("Binaries/Win64") / DefaultVersionSelectorName;
 
 	// Defer to UnrealVersionSelector.exe in a launcher installation if it's got the same version number or greater.
 	FString InstallDir;
