Updating Mac editor behavior to take into the account the newly renamed launcher exe.

diff --git a/Engine/Source/Developer/DesktopPlatform/Private/Mac/DesktopPlatformMac.cpp b/Engine/Source/Developer/DesktopPlatform/Private/Mac/DesktopPlatformMac.cpp
--- a/Engine/Source/Developer/DesktopPlatform/Private/Mac/DesktopPlatformMac.cpp
+++ b/Engine/Source/Developer/DesktopPlatform/Private/Mac/DesktopPlatformMac.cpp
@@ -694,7 +694,7 @@ bool FDesktopPlatformMac::GetLauncherPath(FString& OutLauncherPath) const
 	// Try the default executable in the binaries directory
 	if (FParse::Param(FCommandLine::Get(), TEXT("Dev")))
 	{
-		FString LauncherPath = FPaths::Combine(*FPaths::EngineDir(), TEXT("Binaries"), TEXT("Mac"), TEXT("UnrealEngineLauncher-Mac-Debug.app/Contents/MacOS/UnrealEngineLauncher-Mac-Debug"));
+		FString LauncherPath = FPaths::Combine(*FPaths::EngineDir(), TEXT("Binaries"), TEXT("Mac"), TEXT("EpicGamesLauncher-Mac-Debug.app/Contents/MacOS/EpicGamesLauncher-Mac-Debug"));
 		if(FPaths::FileExists(LauncherPath))
 		{
 			OutLauncherPath = LauncherPath;
@@ -703,7 +703,7 @@ bool FDesktopPlatformMac::GetLauncherPath(FString& OutLauncherPath) const
 	}
 	else
 	{
-		FString LauncherPath = TEXT("/Applications/Epic Games Launcher.app/Contents/MacOS/UnrealEngineLauncher-Mac-Shipping");
+		FString LauncherPath = TEXT("/Applications/Epic Games Launcher.app/Contents/MacOS/EpicGamesLauncher-Mac-Shipping");
 		if(FPaths::FileExists(LauncherPath))
 		{
 			OutLauncherPath = LauncherPath;
