Fixing backwards compatibility issues with project and engine handling in DasktopPlatformBase
Merging using UE4-To-UE4-LauncherDev

diff --git a/Engine/Source/Developer/DesktopPlatform/Private/DesktopPlatformBase.cpp b/Engine/Source/Developer/DesktopPlatform/Private/DesktopPlatformBase.cpp
--- a/Engine/Source/Developer/DesktopPlatform/Private/DesktopPlatformBase.cpp
+++ b/Engine/Source/Developer/DesktopPlatform/Private/DesktopPlatformBase.cpp
@@ -698,7 +698,7 @@ FString FDesktopPlatformBase::GetDefaultProjectCreationPath()
 
 void FDesktopPlatformBase::ReadLauncherInstallationList()
 {
-	FString InstalledListFile = FString(FPlatformProcess::ApplicationSettingsDir()) / TEXT("UnrealEngineLauncher/LauncherInstalled.dat");
+	FString InstalledListFile = FString(FPlatformProcess::ApplicationSettingsDir()) / TEXT("EpicGamesLauncher/LauncherInstalled.dat");
 
 	// If the file does not exist, manually check for the 4.0 or 4.1 manifest
 	FDateTime NewListTimestamp = IFileManager::Get().GetTimeStamp(*InstalledListFile);
@@ -744,7 +744,7 @@ void FDesktopPlatformBase::ReadLauncherInstallationList()
 void FDesktopPlatformBase::CheckForLauncherEngineInstallation(const FString &AppId, const FString &Identifier, TMap<FString, FString> &OutInstallations)
 {
 	FString ManifestText;
-	FString ManifestFileName = FString(FPlatformProcess::ApplicationSettingsDir()) / FString::Printf(TEXT("UnrealEngineLauncher/Data/Manifests/%s.manifest"), *AppId);
+		FString ManifestFileName = FString(FPlatformProcess::ApplicationSettingsDir()) / FString::Printf(TEXT("EpicGamesLauncher/Data/Manifests/%s.manifest"), *AppId);
 	if (FFileHelper::LoadFileToString(ManifestText, *ManifestFileName))
 	{
 		TSharedPtr< FJsonObject > RootObject;
@@ -913,6 +913,13 @@ bool FDesktopPlatformBase::EnumerateProjectsKnownByEngine(const FString &Identif
 
 	// Find the editor game-agnostic settings
 	FConfigSection* Section = GameAgnosticConfig.Find(TEXT("/Script/UnrealEd.EditorSettings"));
+
+	if (Section == NULL)
+	{
+		FConfigCacheIni::LoadExternalIniFile(GameAgnosticConfig, TEXT("EditorGameAgnostic"), NULL, *GameAgnosticConfigDir, false);
+		Section = GameAgnosticConfig.Find(TEXT("/Script/UnrealEd.EditorGameAgnosticSettings"));
+	}
+
 	if(Section != NULL)
 	{
 		// Add in every path that the user has ever created a project file. This is to catch new projects showing up in the user's project folders
