Fixed missing logic from bad integrate.

diff --git a/Engine/Source/Developer/DesktopPlatform/Private/DesktopPlatformBase.cpp b/Engine/Source/Developer/DesktopPlatform/Private/DesktopPlatformBase.cpp
--- a/Engine/Source/Developer/DesktopPlatform/Private/DesktopPlatformBase.cpp
+++ b/Engine/Source/Developer/DesktopPlatform/Private/DesktopPlatformBase.cpp
@@ -700,6 +700,11 @@ void FDesktopPlatformBase::ReadLauncherInstallationList()
 {
 	FString InstalledListFile = FString(FPlatformProcess::ApplicationSettingsDir()) / TEXT("EpicGamesLauncher/LauncherInstalled.dat");
 
+	if (!FPaths::FileExists(InstalledListFile))
+	{
+		InstalledListFile = FString(FPlatformProcess::ApplicationSettingsDir()) / TEXT("UnrealEngineLauncher/LauncherInstalled.dat");
+	}
+
 	// If the file does not exist, manually check for the 4.0 or 4.1 manifest
 	FDateTime NewListTimestamp = IFileManager::Get().GetTimeStamp(*InstalledListFile);
 	if(NewListTimestamp == FDateTime::MinValue())
