Merging using UE4-To-UE4-LauncherDev

diff --git a/Engine/Source/Developer/DesktopPlatform/Private/DesktopPlatformBase.cpp b/Engine/Source/Developer/DesktopPlatform/Private/DesktopPlatformBase.cpp
--- a/Engine/Source/Developer/DesktopPlatform/Private/DesktopPlatformBase.cpp
+++ b/Engine/Source/Developer/DesktopPlatform/Private/DesktopPlatformBase.cpp
@@ -749,7 +749,13 @@ void FDesktopPlatformBase::ReadLauncherInstallationList()
 void FDesktopPlatformBase::CheckForLauncherEngineInstallation(const FString &AppId, const FString &Identifier, TMap<FString, FString> &OutInstallations)
 {
 	FString ManifestText;
-		FString ManifestFileName = FString(FPlatformProcess::ApplicationSettingsDir()) / FString::Printf(TEXT("EpicGamesLauncher/Data/Manifests/%s.manifest"), *AppId);
+	FString ManifestFileName = FString(FPlatformProcess::ApplicationSettingsDir()) / FString::Printf(TEXT("EpicGamesLauncher/Data/Manifests/%s.manifest"), *AppId);
+
+	if (!FPaths::FileExists(ManifestFileName))
+	{
+		ManifestFileName = FString(FPlatformProcess::ApplicationSettingsDir()) / FString::Printf(TEXT("UnrealEngineLauncher/Data/Manifests/%s.manifest"), *AppId);
+	}
+
 	if (FFileHelper::LoadFileToString(ManifestText, *ManifestFileName))
 	{
 		TSharedPtr< FJsonObject > RootObject;
