Merging using UE4-To-UE4-LauncherDev

diff --git a/Engine/Source/Developer/DesktopPlatform/Private/DesktopPlatformBase.cpp b/Engine/Source/Developer/DesktopPlatform/Private/DesktopPlatformBase.cpp
--- a/Engine/Source/Developer/DesktopPlatform/Private/DesktopPlatformBase.cpp
+++ b/Engine/Source/Developer/DesktopPlatform/Private/DesktopPlatformBase.cpp
@@ -698,12 +698,7 @@ FString FDesktopPlatformBase::GetDefaultProjectCreationPath()
 
 void FDesktopPlatformBase::ReadLauncherInstallationList()
 {
-	FString InstalledListFile = FString(FPlatformProcess::ApplicationSettingsDir()) / TEXT("EpicGamesLauncher/LauncherInstalled.dat");
-
-	if (!FPaths::FileExists(InstalledListFile))
-	{
-		InstalledListFile = FString(FPlatformProcess::ApplicationSettingsDir()) / TEXT("UnrealEngineLauncher/LauncherInstalled.dat");
-	}
+	FString InstalledListFile = FString(FPlatformProcess::ApplicationSettingsDir()) / TEXT("UnrealEngineLauncher/LauncherInstalled.dat");
 
 	// If the file does not exist, manually check for the 4.0 or 4.1 manifest
 	FDateTime NewListTimestamp = IFileManager::Get().GetTimeStamp(*InstalledListFile);
@@ -749,12 +744,7 @@ void FDesktopPlatformBase::ReadLauncherInstallationList()
 void FDesktopPlatformBase::CheckForLauncherEngineInstallation(const FString &AppId, const FString &Identifier, TMap<FString, FString> &OutInstallations)
 {
 	FString ManifestText;
-	FString ManifestFileName = FString(FPlatformProcess::ApplicationSettingsDir()) / FString::Printf(TEXT("EpicGamesLauncher/Data/Manifests/%s.manifest"), *AppId);
-
-	if (!FPaths::FileExists(ManifestFileName))
-	{
-		ManifestFileName = FString(FPlatformProcess::ApplicationSettingsDir()) / FString::Printf(TEXT("UnrealEngineLauncher/Data/Manifests/%s.manifest"), *AppId);
-	}
+	FString ManifestFileName = FString(FPlatformProcess::ApplicationSettingsDir()) / FString::Printf(TEXT("UnrealEngineLauncher/Data/Manifests/%s.manifest"), *AppId);
 
 	if (FFileHelper::LoadFileToString(ManifestText, *ManifestFileName))
 	{
