Add a function to IDesktopPlatform to query the version of an engine (either from the new JSON build.version file, or from Version.h if that fails). Make UE4Build.cs update it during builds.

diff --git a/Engine/Build/Build.version b/Engine/Build/Build.version
new file mode 100644
--- /dev/null
+++ b/Engine/Build/Build.version
@@ -0,0 +1,8 @@
+{
+	"MajorVersion": 4,
+	"MinorVersion": 9,
+	"PatchVersion": 0,
+	"Changelist": 2616747,
+	"IsLicenseeVersion": 1,
+	"BranchName": "++depot+UE4"
+}
\ No newline at end of file
diff --git a/Engine/Build/InstalledEngineFilters.ini b/Engine/Build/InstalledEngineFilters.ini
--- a/Engine/Build/InstalledEngineFilters.ini
+++ b/Engine/Build/InstalledEngineFilters.ini
@@ -163,6 +163,7 @@
 /Engine/Documentation/CHM/API.chm
 
 ; Build files
+/Engine/Build/Build.version
 /Engine/Build/Target.cs.template
 
 ; Starter content
diff --git a/Engine/Source/Developer/DesktopPlatform/Private/DesktopPlatformBase.cpp b/Engine/Source/Developer/DesktopPlatform/Private/DesktopPlatformBase.cpp
--- a/Engine/Source/Developer/DesktopPlatform/Private/DesktopPlatformBase.cpp
+++ b/Engine/Source/Developer/DesktopPlatform/Private/DesktopPlatformBase.cpp
@@ -191,6 +191,159 @@ bool FDesktopPlatformBase::IsPreferredEngineIdentifier(const FString &Identifier
 	}
 }
 
+bool FDesktopPlatformBase::TryGetEngineVersion(const FString& RootDir, FEngineVersion& OutVersion)
+{
+	// Read the file to a string
+	FString VersionText;
+	if(FFileHelper::LoadFileToString(VersionText, *(RootDir / TEXT("Engine/Build/Build.version"))))
+	{
+		// Deserialize a JSON object from the string
+		TSharedPtr< FJsonObject > Object;
+		TSharedRef< TJsonReader<> > Reader = TJsonReaderFactory<>::Create(VersionText);
+		if(FJsonSerializer::Deserialize(Reader, Object) && Object.IsValid())
+		{
+			int32 MajorVersion = 0;
+			int32 MinorVersion = 0;
+			int32 PatchVersion = 0;
+			if(Object->TryGetNumberField(TEXT("MajorVersion"), MajorVersion) && Object->TryGetNumberField(TEXT("MinorVersion"), MinorVersion) && Object->TryGetNumberField(TEXT("PatchVersion"), PatchVersion))
+			{
+				int32 Changelist = 0;
+				if(!Object->TryGetNumberField(TEXT("Changelist"), Changelist))
+				{
+					Changelist = 0;
+				}
+
+				int32 IsLicenseeVersion = 0;
+				if(!Object->TryGetNumberField(TEXT("IsLicenseeVersion"), IsLicenseeVersion))
+				{
+					IsLicenseeVersion = 0;
+				}
+
+				FString BranchName;
+				if(!Object->TryGetStringField(TEXT("BranchName"), BranchName))
+				{
+					BranchName = FString();
+				}
+
+				int EncodedChangelist = (IsLicenseeVersion == 0)? Changelist : FEngineVersionBase::EncodeLicenseeChangelist(Changelist);
+				OutVersion = FEngineVersion(MajorVersion, MinorVersion, PatchVersion, EncodedChangelist, BranchName);
+				return true;
+			}
+		}
+	}
+
+	// Try to read the version file
+	FString VersionHeader;
+	if(FFileHelper::LoadFileToString(VersionHeader, *(RootDir / TEXT("Engine/Source/Runtime/Launch/Resources/Version.h"))))
+	{
+		int32 MajorVersion = -1;
+		int32 MinorVersion = -1;
+		int32 PatchVersion = -1;
+		int32 Changelist = 0;
+		int32 IsLicenseeVersion = 0;
+		FString BranchName;
+
+		// Scan the file for version defines
+		const TCHAR* TextPos = *VersionHeader;
+		while(*TextPos)
+		{
+			// Skip over any newlines
+			while(FChar::IsWhitespace(*TextPos))
+			{
+				TextPos++;
+			}
+
+			// Buffer up a line of tokens
+			TArray<FString> Tokens;
+			while(*TextPos != '\n' && *TextPos != 0)
+			{
+				if(*TextPos == ' ' || *TextPos == '\t' || *TextPos == '\r')
+				{
+					// Skip over whitespace
+					TextPos++;
+				}
+				else if(FChar::IsIdentifier(*TextPos))
+				{
+					// Parse an identifier. Exact C rules for an identifier don't really matter; we just need alphanumeric sequences.
+					const TCHAR* TokenStart = TextPos++;
+					while(FChar::IsIdentifier(*TextPos)) TextPos++;
+					Tokens.Add(FString(TextPos - TokenStart, TokenStart));
+				}
+				else if(*TextPos == '\"')
+				{
+					// Parse a string
+					const TCHAR* TokenStart = TextPos++;
+					while(*TextPos != 0 && (TextPos == TokenStart + 1 || *(TextPos - 1) != '\"')) TextPos++;
+					Tokens.Add(FString(TextPos - TokenStart, TokenStart));
+				}
+				else if(*TextPos == '/' && *(TextPos + 1) == '/')
+				{
+					// Skip a C++ style comment
+					TextPos += 2;
+					while(*TextPos != '\n' && *TextPos != 0) TextPos++;
+				}
+				else if(*TextPos == '/' && *(TextPos + 1) == '*' && *(TextPos + 2) != 0 && *(TextPos + 3) != 0)
+				{
+					// Skip a C-style comment
+					TextPos += 4;
+					while(*TextPos != 0 && (*(TextPos - 2) != '*' || *(TextPos - 1) != '/')) TextPos++;
+				}
+				else
+				{
+					// Take a single symbol character
+					Tokens.Add(FString(1, TextPos));
+					TextPos++;
+				}
+			}
+
+			// Check if it matches any version defines
+			if(Tokens.Num() >= 4 && Tokens[0] == "#" && Tokens[1] == "define")
+			{
+				if(FChar::IsDigit(Tokens[3][0]))
+				{
+					if(Tokens[2] == "ENGINE_MAJOR_VERSION")
+					{
+						MajorVersion = FCString::Atoi(*Tokens[3]);
+					}
+					else if(Tokens[2] == "ENGINE_MINOR_VERSION")
+					{
+						MinorVersion = FCString::Atoi(*Tokens[3]);
+					}
+					else if(Tokens[2] == "ENGINE_PATCH_VERSION")
+					{
+						PatchVersion = FCString::Atoi(*Tokens[3]);
+					}
+					else if(Tokens[2] == "ENGINE_VERSION")
+					{
+						Changelist = FCString::Atoi(*Tokens[3]);
+					}
+					else if(Tokens[2] == "ENGINE_IS_LICENSEE_VERSION")
+					{
+						IsLicenseeVersion = FCString::Atoi(*Tokens[3]);
+					}
+				}
+				else if(Tokens[3].StartsWith("\"") && Tokens[3].EndsWith("\""))
+				{
+					if(Tokens[2] == "BRANCH_NAME")
+					{
+						BranchName = Tokens[3].TrimQuotes();
+					}
+				}
+			}
+		}
+
+		// If we have everything we need, fill in the version struct
+		if(MajorVersion != -1 && MinorVersion != -1 && PatchVersion != -1)
+		{
+			int EncodedChangelist = (IsLicenseeVersion == 0)? Changelist : FEngineVersionBase::EncodeLicenseeChangelist(Changelist);
+			OutVersion = FEngineVersion(MajorVersion, MinorVersion, PatchVersion, EncodedChangelist, BranchName);
+			return true;
+		}
+	}
+
+	return false;
+}
+
 bool FDesktopPlatformBase::IsStockEngineRelease(const FString &Identifier)
 {
 	FGuid Guid;
diff --git a/Engine/Source/Developer/DesktopPlatform/Private/DesktopPlatformBase.h b/Engine/Source/Developer/DesktopPlatform/Private/DesktopPlatformBase.h
--- a/Engine/Source/Developer/DesktopPlatform/Private/DesktopPlatformBase.h
+++ b/Engine/Source/Developer/DesktopPlatform/Private/DesktopPlatformBase.h
@@ -26,6 +26,7 @@ public:
 	virtual FString GetEngineSavedConfigDirectory(const FString& Identifier) override;
 	virtual bool IsPreferredEngineIdentifier(const FString &Identifier, const FString &OtherIdentifier) override;
 
+	virtual bool TryGetEngineVersion(const FString& RootDir, FEngineVersion& OutVersion) override;
 	virtual bool IsStockEngineRelease(const FString &Identifier) override;
 	virtual bool TryParseStockEngineVersion(const FString& Identifier, FEngineVersion& OutVersion) override;
 	virtual bool IsSourceDistribution(const FString &RootDir) override;
diff --git a/Engine/Source/Developer/DesktopPlatform/Public/IDesktopPlatform.h b/Engine/Source/Developer/DesktopPlatform/Public/IDesktopPlatform.h
--- a/Engine/Source/Developer/DesktopPlatform/Public/IDesktopPlatform.h
+++ b/Engine/Source/Developer/DesktopPlatform/Public/IDesktopPlatform.h
@@ -243,6 +243,17 @@ public:
 	*/
 	virtual FString GetEngineSavedConfigDirectory(const FString& Identifier) = 0;
 
+	/**
+	 * Attempt to get the engine version from the supplied engine root directory.
+	 * Tries to read the contents of Engine/Build/Build.version, but scrapes information 
+	 * from the Engine/Source/Launch/Resources/Version.h if it fails.
+	 *
+	 * @param	RootDir				Root directory for the engine to check
+	 * @param	OutVersion			Engine version obtained from identifier
+	 * @return	true if the engine version was successfully obtained
+	 */
+	virtual bool TryGetEngineVersion(const FString& RootDir, FEngineVersion& OutVersion) = 0;
+
 	/**
 	* Checks if the given engine identifier is for an stock engine release.
 	*
diff --git a/Engine/Source/Programs/AutomationTool/AutomationUtils/UE4Build.cs b/Engine/Source/Programs/AutomationTool/AutomationUtils/UE4Build.cs
--- a/Engine/Source/Programs/AutomationTool/AutomationUtils/UE4Build.cs
+++ b/Engine/Source/Programs/AutomationTool/AutomationUtils/UE4Build.cs
@@ -5,6 +5,7 @@ using System.Text;
 using System.IO;
 using System.Xml;
 using System.Diagnostics;
+using System.Web.Script.Serialization;
 
 namespace AutomationTool
 {
@@ -483,7 +484,37 @@ namespace AutomationTool
 				}
 				Result.Add(VerFile);
 			}
+			{
+				string VerFile = CombinePaths(CmdEnv.LocalRoot, "Engine", "Build", "build.version");
+				if (bDoUpdateVersionFiles)
+				{
+					Log("Updating {0} with:", VerFile);
+					Log("  Changelist={0}", ChangelistNumber);
+					Log("  IsLicenseeVersion={0}", bIsLicenseeVersion? 1 : 0);
+					Log("  BranchName={0}", Branch);
+
+					string VerText = CommandUtils.ReadAllText(VerFile);
 
+					JavaScriptSerializer Serializer = new JavaScriptSerializer();
+					Dictionary<string, object> Pairs = Serializer.Deserialize<Dictionary<string, object>>(VerText);
+					Pairs["Changelist"] = ChangelistNumber;
+					Pairs["IsLicenseeVersion"] = bIsLicenseeVersion? 1 : 0;
+					Pairs["BranchName"] = Branch;
+					VerText = Serializer.Serialize(Pairs).Replace("{\"", "{\n\t\"").Replace(",\"", ",\n\t\"").Replace("\":", "\": ").Replace("\"}", "\"\n}");
+
+					CommandUtils.WriteAllText(VerFile, VerText);
+				}
+				else
+				{
+					Log("{0} will not be updated because P4 is not enabled.", VerFile);
+				}
+				if (IsBuildMachine)
+				{
+					// Only BuildMachines can check this file in
+					AddBuildProduct(VerFile);
+				}
+				Result.Add(VerFile);
+			}
 			{
 				string VerFile = CombinePaths(CmdEnv.LocalRoot, "Engine", "Source", "Runtime", "Core", "Private", "UObject", "ObjectVersion.cpp");
 				if (bDoUpdateVersionFiles)
diff --git a/Engine/Source/Runtime/Core/Private/Misc/EngineVersion.cpp b/Engine/Source/Runtime/Core/Private/Misc/EngineVersion.cpp
--- a/Engine/Source/Runtime/Core/Private/Misc/EngineVersion.cpp
+++ b/Engine/Source/Runtime/Core/Private/Misc/EngineVersion.cpp
@@ -81,6 +81,11 @@ EVersionComparison FEngineVersionBase::GetNewest(const FEngineVersionBase &First
 	return EVersionComparison::Neither;
 }
 
+uint32 FEngineVersionBase::EncodeLicenseeChangelist(uint32 Changelist)
+{
+	return Changelist | 0x80000000;
+}
+
 
 FEngineVersion::FEngineVersion()
 {
diff --git a/Engine/Source/Runtime/Core/Public/Misc/EngineVersionBase.h b/Engine/Source/Runtime/Core/Public/Misc/EngineVersionBase.h
--- a/Engine/Source/Runtime/Core/Public/Misc/EngineVersionBase.h
+++ b/Engine/Source/Runtime/Core/Public/Misc/EngineVersionBase.h
@@ -66,6 +66,9 @@ public:
 	/** Returns the newest of two versions, and the component at which they differ */
 	static EVersionComparison GetNewest(const FEngineVersionBase &First, const FEngineVersionBase &Second, EVersionComponent *OutComponent);
 
+	/** Encodes a licensee changelist number (by setting the top bit) */
+	static uint32 EncodeLicenseeChangelist(uint32 Changelist);
+
 protected:
 
 	/** Major version number. */
