Check for com.epicgames.EpicGamesLauncher bundle identifier when detecting if the Mac launcher is already running.

diff --git a/Engine/Source/Developer/DesktopPlatform/Private/Mac/DesktopPlatformMac.cpp b/Engine/Source/Developer/DesktopPlatform/Private/Mac/DesktopPlatformMac.cpp
--- a/Engine/Source/Developer/DesktopPlatform/Private/Mac/DesktopPlatformMac.cpp
+++ b/Engine/Source/Developer/DesktopPlatform/Private/Mac/DesktopPlatformMac.cpp
@@ -399,7 +399,11 @@ bool FDesktopPlatformMac::CanOpenLauncher(bool Install)
 bool FDesktopPlatformMac::OpenLauncher(bool Install, FString CommandLineParams )
 {
 	// If the launcher is already running, bring it to front
-	NSArray* RunningLaunchers = [NSRunningApplication runningApplicationsWithBundleIdentifier: @"com.epicgames.UnrealEngineLauncher"];
+	NSArray* RunningLaunchers = [NSRunningApplication runningApplicationsWithBundleIdentifier: @"com.epicgames.EpicGamesLauncher"];
+	if ([RunningLaunchers count] == 0)
+	{
+		RunningLaunchers = [NSRunningApplication runningApplicationsWithBundleIdentifier: @"com.epicgames.UnrealEngineLauncher"];
+	}
 	if ([RunningLaunchers count] > 0)
 	{
 		NSRunningApplication* Launcher = [RunningLaunchers objectAtIndex: 0];
