Updated code that checks for the launcher's existence, opens it and issues navigation commands to use Uri's.

diff --git a/Engine/Source/Developer/DesktopPlatform/Private/Linux/DesktopPlatformLinux.cpp b/Engine/Source/Developer/DesktopPlatform/Private/Linux/DesktopPlatformLinux.cpp
--- a/Engine/Source/Developer/DesktopPlatform/Private/Linux/DesktopPlatformLinux.cpp
+++ b/Engine/Source/Developer/DesktopPlatform/Private/Linux/DesktopPlatformLinux.cpp
@@ -179,7 +179,7 @@ bool FDesktopPlatformLinux::CanOpenLauncher(bool Install)
 	return false;
 }
 
-bool FDesktopPlatformLinux::OpenLauncher(bool Install, FString CommandLineParams )
+bool FDesktopPlatformLinux::OpenLauncher(bool Install, FString LauncherRelativeUrl, FString CommandLineParams)
 {
 	// TODO: support launcher for realz
 	return true;
diff --git a/Engine/Source/Developer/DesktopPlatform/Private/Mac/DesktopPlatformMac.cpp b/Engine/Source/Developer/DesktopPlatform/Private/Mac/DesktopPlatformMac.cpp
--- a/Engine/Source/Developer/DesktopPlatform/Private/Mac/DesktopPlatformMac.cpp
+++ b/Engine/Source/Developer/DesktopPlatform/Private/Mac/DesktopPlatformMac.cpp
@@ -354,49 +354,13 @@ bool FDesktopPlatformMac::OpenFontDialog(const void* ParentWindowHandle, FString
 	return bSuccess;
 }
 
-void OpenLauncherCommandLine( const FString& InCommandLine )
-{
-	FString CommandLine = InCommandLine;
-	FString TransferFilePath = FPaths::Combine(FPlatformProcess::UserSettingsDir(), TEXT("UnrealEngineLauncher"), TEXT("com"), TEXT("transfer.tmp"));
-
-	IFileManager& FileManager = IFileManager::Get();
-	FArchive* TransferFile = NULL;
-	int32 RetryCount = 0;
-
-	while (RetryCount < 2000)
-	{
-		TransferFile = FileManager.CreateFileWriter(*TransferFilePath, FILEWRITE_EvenIfReadOnly);
-
-		if (TransferFile != nullptr)
-		{
-			break;
-		}
-
-		++RetryCount;
-		FPlatformProcess::Sleep(0.01f);
-	}
-
-	if (TransferFile == NULL)
-	{
-		TransferFile = FileManager.CreateFileWriter(*TransferFilePath, FILEWRITE_EvenIfReadOnly | FILEWRITE_NoFail);
-	}
-
-	if (TransferFile != NULL)
-	{
-		*TransferFile << CommandLine;
-		TransferFile->Close();
-		delete TransferFile;
-	}
-	
-}
-
 bool FDesktopPlatformMac::CanOpenLauncher(bool Install)
 {
 	FString Path;
-	return GetLauncherPath(Path) || (Install && GetLauncherInstallerPath(Path));
+	return IsLauncherInstalled() || (Install && GetLauncherInstallerPath(Path));
 }
 
-bool FDesktopPlatformMac::OpenLauncher(bool Install, FString CommandLineParams )
+bool FDesktopPlatformMac::OpenLauncher(bool Install, FString LauncherRelativeUrl, FString CommandLineParams)
 {
 	// If the launcher is already running, bring it to front
 	NSArray* RunningLaunchers = [NSRunningApplication runningApplicationsWithBundleIdentifier: @"com.epicgames.EpicGamesLauncher"];
@@ -404,31 +368,36 @@ bool FDesktopPlatformMac::OpenLauncher(bool Install, FString CommandLineParams )
 	{
 		RunningLaunchers = [NSRunningApplication runningApplicationsWithBundleIdentifier: @"com.epicgames.UnrealEngineLauncher"];
 	}
+
 	if ([RunningLaunchers count] > 0)
 	{
 		NSRunningApplication* Launcher = [RunningLaunchers objectAtIndex: 0];
 		if (!Launcher.hidden || Install || CommandLineParams.Len() > 0) // If the launcher is running, but hidden, don't activate on editor startup
 		{
-			[Launcher activateWithOptions: NSApplicationActivateAllWindows | NSApplicationActivateIgnoringOtherApps];
-			OpenLauncherCommandLine(CommandLineParams); // create a temp file that will tell running Launcher instance to switch to Marketplace tab
+			[Launcher activateWithOptions : NSApplicationActivateAllWindows | NSApplicationActivateIgnoringOtherApps];
 		}
-		return true;
 	}
 
-	// Try to start a new instance
-	FString LauncherPath;
-	if(GetLauncherPath(LauncherPath))
+	if (IsLauncherInstalled())
 	{
+		if (LauncherRelativeUrl.IsEmpty())
+		{
+			LauncherRelativeUrl = TEXT("ue");
+		}
+
 		if (FParse::Param(FCommandLine::Get(), TEXT("Dev")))
 		{
 			CommandLineParams += TEXT(" -noselfupdate");
 		}
-		return FPlatformProcess::CreateProc(*LauncherPath, *CommandLineParams, true, false, false, NULL, 0, *FPaths::GetPath(LauncherPath), NULL).IsValid();
+
+		FString Error;
+		FPlatformProcess::LaunchURL(*FString::Printf(TEXT("com.epicgames.launcher://%s"), *LauncherRelativeUrl), *CommandLineParams, &Error);
+		return true;
 	}
 
 	// Try to install it
 	FString InstallerPath;
-	if(GetLauncherInstallerPath(InstallerPath))
+	if (GetLauncherInstallerPath(InstallerPath))
 	{
 		FPlatformProcess::LaunchFileInDefaultExternalApplication(*InstallerPath);
 		return true;
@@ -695,34 +664,13 @@ bool FDesktopPlatformMac::IsUnrealBuildToolRunning()
 	return FPlatformProcess::IsApplicationRunning(TEXT("mono"));
 }
 
-bool FDesktopPlatformMac::GetLauncherPath(FString& OutLauncherPath) const
+bool FDesktopPlatformMac::IsLauncherInstalled() const
 {
-	// Try the default executable in the binaries directory
-	if (FParse::Param(FCommandLine::Get(), TEXT("Dev")))
-	{
-		FString LauncherPath = FPaths::Combine(*FPaths::EngineDir(), TEXT("Binaries"), TEXT("Mac"), TEXT("EpicGamesLauncher-Mac-Debug.app/Contents/MacOS/EpicGamesLauncher-Mac-Debug"));
-		if(FPaths::FileExists(LauncherPath))
-		{
-			OutLauncherPath = LauncherPath;
-			return true;
-		}
-	}
-	else
-	{
-		FString LauncherPath = TEXT("/Applications/Epic Games Launcher.app/Contents/MacOS/EpicGamesLauncher-Mac-Shipping");
-		if(FPaths::FileExists(LauncherPath))
-		{
-			OutLauncherPath = LauncherPath;
-			return true;
-		}
-	}
-
 	// Otherwise search for it...
 	NSWorkspace* Workspace = [NSWorkspace sharedWorkspace];
 	NSString* Path = [Workspace fullPathForApplication:@"Epic Games Launcher"];
 	if( Path )
 	{
-		OutLauncherPath = FString(Path);
 		return true;
 	}
 
@@ -730,7 +678,6 @@ bool FDesktopPlatformMac::GetLauncherPath(FString& OutLauncherPath) const
 	Path = [Workspace fullPathForApplication:@"Unreal Engine"];
 	if( Path )
 	{
-		OutLauncherPath = FString(Path);
 		return true;
 	}
 
diff --git a/Engine/Source/Developer/DesktopPlatform/Private/Mac/DesktopPlatformMac.h b/Engine/Source/Developer/DesktopPlatform/Private/Mac/DesktopPlatformMac.h
--- a/Engine/Source/Developer/DesktopPlatform/Private/Mac/DesktopPlatformMac.h
+++ b/Engine/Source/Developer/DesktopPlatform/Private/Mac/DesktopPlatformMac.h
@@ -14,7 +14,7 @@ public:
 	virtual bool OpenDirectoryDialog(const void* ParentWindowHandle, const FString& DialogTitle, const FString& DefaultPath, FString& OutFolderName) override;
 	virtual bool OpenFontDialog(const void* ParentWindowHandle, FString& OutFontName, float& OutHeight, EFontImportFlags& OutFlags) override;
 	virtual bool CanOpenLauncher(bool Install) override;
-	virtual bool OpenLauncher(bool Install, FString CommandLineParams ) override;
+	virtual bool OpenLauncher(bool Install, FString LauncherRelativeUrl, FString CommandLineParams) override;
 
 	virtual bool RegisterEngineInstallation(const FString &RootDir, FString &OutIdentifier) override;
 	virtual void EnumerateEngineInstallations(TMap<FString, FString> &OutInstallations) override;
@@ -30,7 +30,7 @@ public:
 	virtual FString GetUserTempPath() override;
 
 private:
-	bool GetLauncherPath(FString& OutLauncherPath) const;
+	bool IsLauncherInstalled() const;
 	bool GetLauncherInstallerPath(FString& OutInstallerPath) const;
 	bool FileDialogShared(bool bSave, const void* ParentWindowHandle, const FString& DialogTitle, const FString& DefaultPath, const FString& DefaultFile, const FString& FileTypes, uint32 Flags, TArray<FString>& OutFilenames, int32& OutFilterIndex);
 };
diff --git a/Engine/Source/Developer/DesktopPlatform/Private/Windows/DesktopPlatformWindows.cpp b/Engine/Source/Developer/DesktopPlatform/Private/Windows/DesktopPlatformWindows.cpp
--- a/Engine/Source/Developer/DesktopPlatform/Private/Windows/DesktopPlatformWindows.cpp
+++ b/Engine/Source/Developer/DesktopPlatform/Private/Windows/DesktopPlatformWindows.cpp
@@ -146,27 +146,37 @@ bool FDesktopPlatformWindows::OpenFontDialog(const void* ParentWindowHandle, FSt
 
 bool FDesktopPlatformWindows::CanOpenLauncher(bool Install)
 {
+	FRegistryRootedKey UriProtocolKey(HKEY_CLASSES_ROOT, TEXT("com.epicgames.launcher"));
+
 	// Check if the launcher exists, or (optionally) if the installer exists
 	FString Path;
-	return GetLauncherPath(Path) || (Install && GetLauncherInstallerPath(Path));
+	return UriProtocolKey.Exists() || (Install && GetLauncherInstallerPath(Path));
 }
 
-bool FDesktopPlatformWindows::OpenLauncher(bool Install, FString CommandLineParams )
+bool FDesktopPlatformWindows::OpenLauncher(bool Install, FString LauncherRelativeUrl, FString CommandLineParams)
 {
+	FRegistryRootedKey UriProtocolKey(HKEY_CLASSES_ROOT, TEXT("com.epicgames.launcher"));
+	FString InstallerPath;
+
 	// Try to launch it directly
-	FString LauncherPath;
-	if(GetLauncherPath(LauncherPath))
+	if (UriProtocolKey.Exists())
 	{
+		if (LauncherRelativeUrl.IsEmpty())
+		{
+			LauncherRelativeUrl = TEXT("ue");
+		}
+
 		if (FParse::Param(FCommandLine::Get(), TEXT("Dev")))
 		{
 			CommandLineParams += TEXT(" -noselfupdate");
 		}
-		return FPlatformProcess::CreateProc(*LauncherPath, *CommandLineParams, true, false, false, NULL, 0, *FPaths::GetPath(LauncherPath), NULL).IsValid();
-	}
 
+		FString Error;
+		FPlatformProcess::LaunchURL(*FString::Printf(TEXT("com.epicgames.launcher://%s"), *LauncherRelativeUrl), *CommandLineParams, &Error);
+		return true;
+	}
 	// Otherwise see if we can install it
-	FString InstallerPath;
-	if(Install && GetLauncherInstallerPath(InstallerPath))
+	else if(Install && GetLauncherInstallerPath(InstallerPath))
 	{
 		FPlatformProcess::LaunchFileInDefaultExternalApplication(*InstallerPath);
 		return true;
@@ -175,52 +185,15 @@ bool FDesktopPlatformWindows::OpenLauncher(bool Install, FString CommandLinePara
 	return false;
 }
 
-bool FDesktopPlatformWindows::GetLauncherPath(FString& OutLauncherPath)
-{
-	// Try the default executable in the binaries directory
-	if (FParse::Param(FCommandLine::Get(), TEXT("Dev")))
-	{
-		FString LauncherPath = FPaths::Combine(*FPaths::EngineDir(), TEXT("Binaries"), TEXT("Win64"), TEXT("UnrealEngineLauncher-Win64-Debug.exe"));
-		if(FPaths::FileExists(LauncherPath))
-		{
-			OutLauncherPath = LauncherPath;
-			return true;
-		}
-	}
-	else
-	{
-		FString LauncherPath = FPaths::Combine(*FPaths::EngineDir(), TEXT("../../Launcher/Engine/Binaries/Win64/UnrealEngineLauncher.exe"));
-		if(FPaths::FileExists(LauncherPath))
-		{
-			OutLauncherPath = LauncherPath;
-			return true;
-		}
-	}
-
-	// Otherwise locate it in the Registry...
-	FString InstallDir;
-	if (FWindowsPlatformMisc::QueryRegKey(HKEY_LOCAL_MACHINE, TEXT("Software\\EpicGames\\Unreal Engine"), TEXT("InstallDir"), InstallDir) && (InstallDir.Len() > 0))
-	{
-		FString LauncherPath = FPaths::Combine(*InstallDir, TEXT("Launcher/Engine/Binaries/Win64/UnrealEngineLauncher.exe"));
-		if(FPaths::FileExists(LauncherPath))
-		{
-			OutLauncherPath = LauncherPath;
-			return true;
-		}
-	}
-
-	// Otherwise fail
-	return false;
-}
-
 bool FDesktopPlatformWindows::GetLauncherInstallerPath(FString& OutInstallerPath)
 {
-	FString InstallerPath = FPaths::ConvertRelativePathToFull(FPaths::Combine(*FPaths::EngineDir(), TEXT("Extras/UnrealEngineLauncher/UnrealEngineInstaller.msi")));
-	if(FPaths::FileExists(InstallerPath))
+	FString InstallerPath = FPaths::ConvertRelativePathToFull(FPaths::Combine(*FPaths::EngineDir(), TEXT("Extras/UnrealEngineLauncher/EpicGamesLauncherInstaller.msi")));
+	if (FPaths::FileExists(InstallerPath))
 	{
 		OutInstallerPath = InstallerPath;
 		return true;
 	}
+
 	return false;
 }
 
diff --git a/Engine/Source/Developer/DesktopPlatform/Private/Windows/DesktopPlatformWindows.h b/Engine/Source/Developer/DesktopPlatform/Private/Windows/DesktopPlatformWindows.h
--- a/Engine/Source/Developer/DesktopPlatform/Private/Windows/DesktopPlatformWindows.h
+++ b/Engine/Source/Developer/DesktopPlatform/Private/Windows/DesktopPlatformWindows.h
@@ -15,7 +15,7 @@ public:
 	virtual bool OpenDirectoryDialog(const void* ParentWindowHandle, const FString& DialogTitle, const FString& DefaultPath, FString& OutFolderName) override;
 	virtual bool OpenFontDialog(const void* ParentWindowHandle, FString& OutFontName, float& OutHeight, EFontImportFlags& OutFlags) override;
 	virtual bool CanOpenLauncher(bool Install) override;
-	virtual bool OpenLauncher(bool Install, FString CommandLineParams ) override;
+	virtual bool OpenLauncher(bool Install, FString LauncherRelativeUrl, FString CommandLineParams) override;
 
 	virtual bool RegisterEngineInstallation(const FString &RootDir, FString &OutIdentifier) override;
 	virtual void EnumerateEngineInstallations(TMap<FString, FString> &OutInstallations) override;
@@ -35,7 +35,6 @@ public:
 	virtual FString GetUserTempPath() override;
 
 private:
-	bool GetLauncherPath(FString& OutLauncherPath);
 	bool GetLauncherInstallerPath(FString& OutInstallerPath);
 	bool FileDialogShared(bool bSave, const void* ParentWindowHandle, const FString& DialogTitle, const FString& DefaultPath, const FString& DefaultFile, const FString& FileTypes, uint32 Flags, TArray<FString>& OutFilenames, int32& OutFilterIndex);
 	void GetRequiredRegistrySettings(TIndirectArray<FRegistryRootedKey> &RootedKeys);
diff --git a/Engine/Source/Developer/DesktopPlatform/Public/IDesktopPlatform.h b/Engine/Source/Developer/DesktopPlatform/Public/IDesktopPlatform.h
--- a/Engine/Source/Developer/DesktopPlatform/Public/IDesktopPlatform.h
+++ b/Engine/Source/Developer/DesktopPlatform/Public/IDesktopPlatform.h
@@ -138,9 +138,11 @@ public:
 	 * Opens the marketplace user interface.
 	 *
 	 * @param Install					Whether to install the marketplace if it is missing.
+	 * @param LauncherRelativeUrl		A url relative to the launcher which you'd like the launcher to navigate to. Empty defaults to the UE homepage
+	 * @param CommandLineParams			Optional command to open the launcher with if it is not already open
 	 * @return true if the marketplace was opened, false if it is not installed or could not be installed/opened.
 	 */
-	virtual bool OpenLauncher(bool Install, FString CommandLineParams ) = 0;
+	virtual bool OpenLauncher(bool Install, FString LauncherRelativeUrl, FString CommandLineParams) = 0;
 
 	/**
 	* Returns a description for the engine with the given identifier.
diff --git a/Engine/Source/Editor/GameProjectGeneration/Private/SProjectBrowser.cpp b/Engine/Source/Editor/GameProjectGeneration/Private/SProjectBrowser.cpp
--- a/Engine/Source/Editor/GameProjectGeneration/Private/SProjectBrowser.cpp
+++ b/Engine/Source/Editor/GameProjectGeneration/Private/SProjectBrowser.cpp
@@ -1230,7 +1230,7 @@ FReply SProjectBrowser::HandleMarketplaceTabButtonClicked()
 	{
 		TArray<FAnalyticsEventAttribute> EventAttributes;
 
-		if (DesktopPlatform->OpenLauncher(false, TEXT("-OpenMarket")))
+		if (DesktopPlatform->OpenLauncher(false, TEXT("ue/marketplace"), FString()))
 		{
 			EventAttributes.Add(FAnalyticsEventAttribute(TEXT("OpenSucceeded"), TEXT("TRUE")));
 		}
@@ -1240,7 +1240,7 @@ FReply SProjectBrowser::HandleMarketplaceTabButtonClicked()
 
 			if (EAppReturnType::Yes == FMessageDialog::Open(EAppMsgType::YesNo, LOCTEXT("InstallMarketplacePrompt", "The Marketplace requires the Epic Games Launcher, which does not seem to be installed on your computer. Would you like to install it now?")))
 			{
-				if (!DesktopPlatform->OpenLauncher(true, TEXT("-OpenMarket")))
+				if (!DesktopPlatform->OpenLauncher(true, TEXT("ue/marketplace"), FString()))
 				{
 					EventAttributes.Add(FAnalyticsEventAttribute(TEXT("InstallSucceeded"), TEXT("FALSE")));
 					FMessageDialog::Open(EAppMsgType::Ok, FText::FromString(TEXT("Sorry, there was a problem installing the Launcher.\nPlease try to install it manually!")));
diff --git a/Engine/Source/Editor/LevelEditor/Private/LevelEditorActions.cpp b/Engine/Source/Editor/LevelEditor/Private/LevelEditorActions.cpp
--- a/Engine/Source/Editor/LevelEditor/Private/LevelEditorActions.cpp
+++ b/Engine/Source/Editor/LevelEditor/Private/LevelEditorActions.cpp
@@ -1961,7 +1961,7 @@ void FLevelEditorActionCallbacks::OpenMarketplace()
 	{
 		TArray<FAnalyticsEventAttribute> EventAttributes;
 
-		if (DesktopPlatform->OpenLauncher(false, TEXT("-OpenMarket")))
+		if (DesktopPlatform->OpenLauncher(false, TEXT("ue/marketplace"), FString()))
 		{
 			EventAttributes.Add(FAnalyticsEventAttribute(TEXT("OpenSucceeded"), TEXT("TRUE")));
 		}
@@ -1971,7 +1971,7 @@ void FLevelEditorActionCallbacks::OpenMarketplace()
 
 			if (EAppReturnType::Yes == FMessageDialog::Open(EAppMsgType::YesNo, LOCTEXT("InstallMarketplacePrompt", "The Marketplace requires the Epic Games Launcher, which does not seem to be installed on your computer. Would you like to install it now?")))
 			{
-				if (!DesktopPlatform->OpenLauncher(true, TEXT("-OpenMarket")))
+				if (!DesktopPlatform->OpenLauncher(true, TEXT("ue/marketplace"), FString()))
 				{
 					EventAttributes.Add(FAnalyticsEventAttribute(TEXT("InstallSucceeded"), TEXT("FALSE")));
 					FMessageDialog::Open(EAppMsgType::Ok, FText::FromString(TEXT("Sorry, there was a problem installing the Launcher.\nPlease try to install it manually!")));
diff --git a/Engine/Source/Editor/NewsFeed/Private/Widgets/SNewsFeed.cpp b/Engine/Source/Editor/NewsFeed/Private/Widgets/SNewsFeed.cpp
--- a/Engine/Source/Editor/NewsFeed/Private/Widgets/SNewsFeed.cpp
+++ b/Engine/Source/Editor/NewsFeed/Private/Widgets/SNewsFeed.cpp
@@ -304,7 +304,7 @@ void SNewsFeed::HandleNewsListViewSelectionChanged( FNewsFeedItemPtr Selection,
 			}
 			else if (Selection->Url.StartsWith(TEXT("ue4://market/")))
 			{
-				Succeeded = DesktopPlatform->OpenLauncher(false, FString::Printf(TEXT("-OpenMarketItem=%s"), *Selection->Url.RightChop(13)));
+				Succeeded = DesktopPlatform->OpenLauncher(false, FString::Printf(TEXT("ue/marketplace/%s"), *Selection->Url.RightChop(13)), FString());
 			}
 //			else if (Selection->Url.StartsWith(TEXT("ue4://")))
 //			{
diff --git a/Engine/Source/Editor/UnrealEd/Private/EditorEngine.cpp b/Engine/Source/Editor/UnrealEd/Private/EditorEngine.cpp
--- a/Engine/Source/Editor/UnrealEd/Private/EditorEngine.cpp
+++ b/Engine/Source/Editor/UnrealEd/Private/EditorEngine.cpp
@@ -443,7 +443,7 @@ void UEditorEngine::InitEditor(IEngineLoop* InEngineLoop)
 		IDesktopPlatform* DesktopPlatform = FDesktopPlatformModule::Get();
 		if ( DesktopPlatform != NULL )
 		{
-			DesktopPlatform->OpenLauncher(false, TEXT(""));
+			DesktopPlatform->OpenLauncher(false, FString(), FString());
 		}
 	}
 
