Build fix

diff --git a/Engine/Source/Developer/DesktopPlatform/Private/Linux/DesktopPlatformLinux.h b/Engine/Source/Developer/DesktopPlatform/Private/Linux/DesktopPlatformLinux.h
--- a/Engine/Source/Developer/DesktopPlatform/Private/Linux/DesktopPlatformLinux.h
+++ b/Engine/Source/Developer/DesktopPlatform/Private/Linux/DesktopPlatformLinux.h
@@ -14,7 +14,7 @@ public:
 	virtual bool OpenDirectoryDialog(const void* ParentWindowHandle, const FString& DialogTitle, const FString& DefaultPath, FString& OutFolderName) override;
 	virtual bool OpenFontDialog(const void* ParentWindowHandle, FString& OutFontName, float& OutHeight, EFontImportFlags& OutFlags) override;
 	virtual bool CanOpenLauncher(bool Install) override;
-	virtual bool OpenLauncher(bool Install, FString CommandLineParams ) override;
+	virtual bool OpenLauncher(bool Install, FString LauncherRelativeUrl, FString CommandLineParams) override;
 
 	virtual bool RegisterEngineInstallation(const FString &RootDir, FString &OutIdentifier) override;
 	virtual void EnumerateEngineInstallations(TMap<FString, FString> &OutInstallations) override;
