Fixed SlateFileDialogs dependency on DesktopPlatform

diff --git a/Engine/Source/Developer/DesktopPlatform/DesktopPlatform.Build.cs b/Engine/Source/Developer/DesktopPlatform/DesktopPlatform.Build.cs
--- a/Engine/Source/Developer/DesktopPlatform/DesktopPlatform.Build.cs
+++ b/Engine/Source/Developer/DesktopPlatform/DesktopPlatform.Build.cs
@@ -17,8 +17,20 @@ public class DesktopPlatform : ModuleRules
 
 		if (Target.Platform == UnrealTargetPlatform.Linux)
 		{
+			PrivateIncludePathModuleNames.AddRange(
+				new string[] {
+					"SlateFileDialogs",
+				}
+			);
+
+			DynamicallyLoadedModuleNames.AddRange(
+				new string[] {
+					"SlateFileDialogs",
+				}
+			);
+
 			AddThirdPartyPrivateStaticDependencies(Target, "SDL2");
-			AddThirdPartyPrivateStaticDependencies(Target, "SlateFileDialogs");
+			//AddThirdPartyPrivateStaticDependencies(Target, "LinuxNativeDialogs");
 		}
 	}
 }
