For an empty request, the launcher scheme should be invoked with no path. This means to just make sure the Launcher is running, without performing a UI request and window focus

diff --git a/Engine/Source/Developer/DesktopPlatform/Private/Mac/DesktopPlatformMac.cpp b/Engine/Source/Developer/DesktopPlatform/Private/Mac/DesktopPlatformMac.cpp
--- a/Engine/Source/Developer/DesktopPlatform/Private/Mac/DesktopPlatformMac.cpp
+++ b/Engine/Source/Developer/DesktopPlatform/Private/Mac/DesktopPlatformMac.cpp
@@ -380,9 +380,14 @@ bool FDesktopPlatformMac::OpenLauncher(bool Install, FString LauncherRelativeUrl
 
 	if (IsLauncherInstalled())
 	{
+		FString LauncherUriRequest;
 		if (LauncherRelativeUrl.IsEmpty())
 		{
-			LauncherRelativeUrl = TEXT("ue");
+			LauncherUriRequest = TEXT("com.epicgames.launcher:");
+		}
+		else
+		{
+			LauncherUriRequest = FString::Printf(TEXT("com.epicgames.launcher://%s"), *LauncherRelativeUrl);
 		}
 
 		if (FParse::Param(FCommandLine::Get(), TEXT("Dev")))
@@ -391,7 +396,7 @@ bool FDesktopPlatformMac::OpenLauncher(bool Install, FString LauncherRelativeUrl
 		}
 
 		FString Error;
-		FPlatformProcess::LaunchURL(*FString::Printf(TEXT("com.epicgames.launcher://%s"), *LauncherRelativeUrl), *CommandLineParams, &Error);
+		FPlatformProcess::LaunchURL(*LauncherUriRequest, *CommandLineParams, &Error);
 		return true;
 	}
 
diff --git a/Engine/Source/Developer/DesktopPlatform/Private/Windows/DesktopPlatformWindows.cpp b/Engine/Source/Developer/DesktopPlatform/Private/Windows/DesktopPlatformWindows.cpp
--- a/Engine/Source/Developer/DesktopPlatform/Private/Windows/DesktopPlatformWindows.cpp
+++ b/Engine/Source/Developer/DesktopPlatform/Private/Windows/DesktopPlatformWindows.cpp
@@ -161,9 +161,14 @@ bool FDesktopPlatformWindows::OpenLauncher(bool Install, FString LauncherRelativ
 	// Try to launch it directly
 	if (UriProtocolKey.Exists())
 	{
+		FString LauncherUriRequest;
 		if (LauncherRelativeUrl.IsEmpty())
 		{
-			LauncherRelativeUrl = TEXT("ue");
+			LauncherUriRequest = TEXT("com.epicgames.launcher:");
+		}
+		else
+		{
+			LauncherUriRequest = FString::Printf(TEXT("com.epicgames.launcher://%s"), *LauncherRelativeUrl);
 		}
 
 		if (FParse::Param(FCommandLine::Get(), TEXT("Dev")))
@@ -172,7 +177,7 @@ bool FDesktopPlatformWindows::OpenLauncher(bool Install, FString LauncherRelativ
 		}
 
 		FString Error;
-		FPlatformProcess::LaunchURL(*FString::Printf(TEXT("com.epicgames.launcher://%s"), *LauncherRelativeUrl), *CommandLineParams, &Error);
+		FPlatformProcess::LaunchURL(*LauncherUriRequest, *CommandLineParams, &Error);
 		return true;
 	}
 	// Otherwise see if we can install it
