Changed the Mac OpenLauncher function to no longer bring the launcher to the front when just opening the editor.

diff --git a/Engine/Source/Developer/DesktopPlatform/Private/Mac/DesktopPlatformMac.cpp b/Engine/Source/Developer/DesktopPlatform/Private/Mac/DesktopPlatformMac.cpp
--- a/Engine/Source/Developer/DesktopPlatform/Private/Mac/DesktopPlatformMac.cpp
+++ b/Engine/Source/Developer/DesktopPlatform/Private/Mac/DesktopPlatformMac.cpp
@@ -362,6 +362,21 @@ bool FDesktopPlatformMac::CanOpenLauncher(bool Install)
 
 bool FDesktopPlatformMac::OpenLauncher(bool Install, FString LauncherRelativeUrl, FString CommandLineParams)
 {
+	FString LauncherUriRequest;
+	if (LauncherRelativeUrl.IsEmpty())
+	{
+		LauncherUriRequest = TEXT("com.epicgames.launcher:");
+	}
+	else
+	{
+		LauncherUriRequest = FString::Printf(TEXT("com.epicgames.launcher://%s"), *LauncherRelativeUrl);
+	}
+
+	if (FParse::Param(FCommandLine::Get(), TEXT("Dev")))
+	{
+		CommandLineParams += TEXT(" -noselfupdate");
+	}
+
 	// If the launcher is already running, bring it to front
 	NSArray* RunningLaunchers = [NSRunningApplication runningApplicationsWithBundleIdentifier: @"com.epicgames.EpicGamesLauncher"];
 	if ([RunningLaunchers count] == 0)
@@ -375,26 +390,14 @@ bool FDesktopPlatformMac::OpenLauncher(bool Install, FString LauncherRelativeUrl
 		if (!Launcher.hidden || Install || CommandLineParams.Len() > 0) // If the launcher is running, but hidden, don't activate on editor startup
 		{
 			[Launcher activateWithOptions : NSApplicationActivateAllWindows | NSApplicationActivateIgnoringOtherApps];
+			FString Error;
+			FPlatformProcess::LaunchURL(*LauncherUriRequest, *CommandLineParams, &Error);
 		}
+		return true;
 	}
 
 	if (IsLauncherInstalled())
 	{
-		FString LauncherUriRequest;
-		if (LauncherRelativeUrl.IsEmpty())
-		{
-			LauncherUriRequest = TEXT("com.epicgames.launcher:");
-		}
-		else
-		{
-			LauncherUriRequest = FString::Printf(TEXT("com.epicgames.launcher://%s"), *LauncherRelativeUrl);
-		}
-
-		if (FParse::Param(FCommandLine::Get(), TEXT("Dev")))
-		{
-			CommandLineParams += TEXT(" -noselfupdate");
-		}
-
 		FString Error;
 		FPlatformProcess::LaunchURL(*LauncherUriRequest, *CommandLineParams, &Error);
 		return true;
