Copying //UE4/Dev-Sequencer to //UE4/Dev-Main (Source: //UE4/Dev-Sequencer @ 3237992)

diff --git a/Engine/Plugins/MovieScene/LevelSequenceEditor/Source/LevelSequenceEditor/Private/LevelSequenceEditorToolkit.cpp b/Engine/Plugins/MovieScene/LevelSequenceEditor/Source/LevelSequenceEditor/Private/LevelSequenceEditorToolkit.cpp
--- a/Engine/Plugins/MovieScene/LevelSequenceEditor/Source/LevelSequenceEditor/Private/LevelSequenceEditorToolkit.cpp
+++ b/Engine/Plugins/MovieScene/LevelSequenceEditor/Source/LevelSequenceEditor/Private/LevelSequenceEditorToolkit.cpp
@@ -42,8 +42,10 @@
 #include "SceneOutlinerPublicTypes.h"
 #include "Widgets/Docking/SDockTab.h"
 #include "SequencerSettings.h"
+#include "LevelEditorSequencerIntegration.h"
 #include "MovieSceneCaptureDialogModule.h"
 #include "MovieScene.h"
+#include "UnrealEdMisc.h"
 
 // @todo sequencer: hack: setting defaults for transform tracks
 #include "Sections/MovieScene3DTransformSection.h"
@@ -161,6 +163,8 @@ FLevelSequenceEditorToolkit::FLevelSequenceEditorToolkit(const TSharedRef<ISlate
 
 FLevelSequenceEditorToolkit::~FLevelSequenceEditorToolkit()
 {
+	FLevelEditorSequencerIntegration::Get().RemoveSequencer(Sequencer.ToSharedRef());
+
 	Sequencer->Close();
 
 	// unregister delegates
@@ -188,7 +192,7 @@ FLevelSequenceEditorToolkit::~FLevelSequenceEditorToolkit()
 /* FLevelSequenceEditorToolkit interface
  *****************************************************************************/
 
-void FLevelSequenceEditorToolkit::Initialize(const EToolkitMode::Type Mode, const TSharedPtr<IToolkitHost>& InitToolkitHost, ULevelSequence* InLevelSequence, bool bEditWithinLevelEditor)
+void FLevelSequenceEditorToolkit::Initialize(const EToolkitMode::Type Mode, const TSharedPtr<IToolkitHost>& InitToolkitHost, ULevelSequence* InLevelSequence)
 {
 	// create tab layout
 	const TSharedRef<FTabManager::FLayout> StandaloneDefaultLayout = FTabManager::NewLayout("Standalone_LevelSequenceEditor")
@@ -215,9 +219,10 @@ void FLevelSequenceEditorToolkit::Initialize(const EToolkitMode::Type Mode, cons
 	FSequencerInitParams SequencerInitParams;
 	{
 		SequencerInitParams.RootSequence = LevelSequence;
-		SequencerInitParams.bEditWithinLevelEditor = bEditWithinLevelEditor;
+		SequencerInitParams.bEditWithinLevelEditor = true;
 		SequencerInitParams.ToolkitHost = InitToolkitHost;
 		SequencerInitParams.SpawnRegister = SpawnRegister;
+
 		SequencerInitParams.EventContexts.BindStatic(GetLevelSequenceEditorEventContexts);
 		SequencerInitParams.PlaybackContext.BindStatic(GetLevelSequenceEditorPlaybackContext);
 
@@ -239,26 +244,31 @@ void FLevelSequenceEditorToolkit::Initialize(const EToolkitMode::Type Mode, cons
 
 		SequencerInitParams.ViewParams.AddMenuExtender = AddMenuExtender;
 		SequencerInitParams.ViewParams.UniqueName = "LevelSequenceEditor";
+		SequencerInitParams.ViewParams.OnReceivedFocus.BindRaw(this, &FLevelSequenceEditorToolkit::OnSequencerReceivedFocus);
 	}
 
 	Sequencer = FModuleManager::LoadModuleChecked<ISequencerModule>("Sequencer").CreateSequencer(SequencerInitParams);
 	SpawnRegister->SetSequencer(Sequencer);
 	Sequencer->OnActorAddedToSequencer().AddSP(this, &FLevelSequenceEditorToolkit::HandleActorAddedToSequencer);
 
-	if (bEditWithinLevelEditor)
-	{
-		// @todo remove when world-centric mode is added
-		FLevelEditorModule& LevelEditorModule = FModuleManager::LoadModuleChecked<FLevelEditorModule>("LevelEditor");
+	FLevelEditorSequencerIntegrationOptions Options;
+	Options.bRequiresLevelEvents = true;
+	Options.bRequiresActorEvents = true;
+	Options.bCanRecord = true;
+
+	FLevelEditorSequencerIntegration::Get().AddSequencer(Sequencer.ToSharedRef(), Options);
 
-		LevelEditorModule.AttachSequencer(Sequencer->GetSequencerWidget(), SharedThis(this));
+	// @todo remove when world-centric mode is added
+	FLevelEditorModule& LevelEditorModule = FModuleManager::LoadModuleChecked<FLevelEditorModule>("LevelEditor");
 
-		// We need to find out when the user loads a new map, because we might need to re-create puppet actors
-		// when previewing a MovieScene
-		LevelEditorModule.OnMapChanged().AddRaw(this, &FLevelSequenceEditorToolkit::HandleMapChanged);
+	LevelEditorModule.AttachSequencer(Sequencer->GetSequencerWidget(), SharedThis(this));
 
-		ILevelSequenceEditorModule& LevelSequenceEditorModule = FModuleManager::LoadModuleChecked<ILevelSequenceEditorModule>("LevelSequenceEditor");
-		LevelSequenceEditorModule.OnMasterSequenceCreated().AddRaw(this, &FLevelSequenceEditorToolkit::HandleMasterSequenceCreated);
-	}
+	// We need to find out when the user loads a new map, because we might need to re-create puppet actors
+	// when previewing a MovieScene
+	LevelEditorModule.OnMapChanged().AddRaw(this, &FLevelSequenceEditorToolkit::HandleMapChanged);
+
+	ILevelSequenceEditorModule& LevelSequenceEditorModule = FModuleManager::LoadModuleChecked<ILevelSequenceEditorModule>("LevelSequenceEditor");
+	LevelSequenceEditorModule.OnMasterSequenceCreated().AddRaw(this, &FLevelSequenceEditorToolkit::HandleMasterSequenceCreated);
 
 	FLevelSequenceEditorToolkit::OnOpened().Broadcast(*this);
 }
@@ -652,6 +662,14 @@ void FLevelSequenceEditorToolkit::AddPosessActorMenuExtensions(FMenuBuilder& Men
 /* FLevelSequenceEditorToolkit callbacks
  *****************************************************************************/
 
+void FLevelSequenceEditorToolkit::OnSequencerReceivedFocus()
+{
+	if (Sequencer.IsValid())
+	{
+		FLevelEditorSequencerIntegration::Get().OnSequencerReceivedFocus(Sequencer.ToSharedRef());
+	}
+}
+
 void FLevelSequenceEditorToolkit::HandleAddComponentActionExecute(UActorComponent* Component)
 {
 	Sequencer->GetHandleToObject(Component);
@@ -686,7 +704,12 @@ void FLevelSequenceEditorToolkit::HandleActorAddedToSequencer(AActor* Actor, con
 
 void FLevelSequenceEditorToolkit::HandleMapChanged(class UWorld* NewWorld, EMapChangeType MapChangeType)
 {
-	Sequencer->NotifyMapChanged(NewWorld, MapChangeType);
+	// @todo sequencer: We should only wipe/respawn puppets that are affected by the world that is being changed! (multi-UWorld support)
+	if( ( MapChangeType == EMapChangeType::LoadMap || MapChangeType == EMapChangeType::NewMap || MapChangeType == EMapChangeType::TearDownWorld) )
+	{
+		Sequencer->GetSpawnRegister().CleanUp(*Sequencer);
+		Sequencer->UpdateRuntimeInstances();
+	}
 }
 
 void FLevelSequenceEditorToolkit::AddShot(UMovieSceneCinematicShotTrack* ShotTrack, const FString& ShotAssetName, const FString& ShotPackagePath, float ShotStartTime, float ShotEndTime, UObject* AssetToDuplicate)
diff --git a/Engine/Plugins/MovieScene/LevelSequenceEditor/Source/LevelSequenceEditor/Private/LevelSequenceEditorToolkit.h b/Engine/Plugins/MovieScene/LevelSequenceEditor/Source/LevelSequenceEditor/Private/LevelSequenceEditorToolkit.h
--- a/Engine/Plugins/MovieScene/LevelSequenceEditor/Source/LevelSequenceEditor/Private/LevelSequenceEditorToolkit.h
+++ b/Engine/Plugins/MovieScene/LevelSequenceEditor/Source/LevelSequenceEditor/Private/LevelSequenceEditorToolkit.h
@@ -62,9 +62,8 @@ public:
 	 * @param InitToolkitHost When Mode is WorldCentric, this is the level editor instance to spawn this editor within.
 	 * @param LevelSequence The animation to edit.
 	 * @param TrackEditorDelegates Delegates to call to create auto-key handlers for this sequencer.
-	 * @param bEditWithinLevelEditor Whether or not sequencer should be edited within the level editor.
 	 */
-	void Initialize(const EToolkitMode::Type Mode, const TSharedPtr<IToolkitHost>& InitToolkitHost, ULevelSequence* LevelSequence, bool bEditWithinLevelEditor);
+	void Initialize(const EToolkitMode::Type Mode, const TSharedPtr<IToolkitHost>& InitToolkitHost, ULevelSequence* LevelSequence);
 
 	/**
 	 * Get the sequencer object being edited in this tool kit.
@@ -112,6 +111,9 @@ protected:
 	/** Add a shot to a master sequence */
 	void AddShot(UMovieSceneCinematicShotTrack* ShotTrack, const FString& ShotAssetName, const FString& ShotPackagePath, float ShotStartTime, float ShotEndTime, UObject* AssetToDuplicate);
 
+	/** Called whenever sequencer has received focus */
+	void OnSequencerReceivedFocus();
+
 private:
 
 	/** Callback for executing the Add Component action. */
diff --git a/Engine/Plugins/MovieScene/MatineeToLevelSequence/Source/MatineeToLevelSequence/Private/MatineeToLevelSequenceModule.cpp b/Engine/Plugins/MovieScene/MatineeToLevelSequence/Source/MatineeToLevelSequence/Private/MatineeToLevelSequenceModule.cpp
--- a/Engine/Plugins/MovieScene/MatineeToLevelSequence/Source/MatineeToLevelSequence/Private/MatineeToLevelSequenceModule.cpp
+++ b/Engine/Plugins/MovieScene/MatineeToLevelSequence/Source/MatineeToLevelSequence/Private/MatineeToLevelSequenceModule.cpp
@@ -38,6 +38,7 @@
 #include "Matinee/InterpTrackDirector.h"
 #include "Matinee/InterpTrackEvent.h"
 #include "Matinee/InterpTrackAudioMaster.h"
+#include "Matinee/InterpTrackVectorProp.h"
 #include "Matinee/InterpTrackVisibility.h"
 #include "Matinee/InterpGroupInst.h"
 #include "Matinee/InterpGroupDirector.h"
@@ -65,6 +66,7 @@
 #include "Tracks/MovieSceneSlomoTrack.h"
 #include "Tracks/MovieSceneCameraCutTrack.h"
 #include "Tracks/MovieSceneEventTrack.h"
+#include "Tracks/MovieSceneVectorTrack.h"
 #include "Tracks/MovieSceneVisibilityTrack.h"
 #include "Animation/SkeletalMeshActor.h"
 #include "Dialogs/Dialogs.h"
@@ -474,8 +476,9 @@ protected:
 
 				if ( bHasKeyframes && PossessableGuid.IsValid())
 				{
+					FVector DefaultScale = GroupActor != nullptr ? GroupActor->GetActorScale() : FVector(1.f);
 					UMovieScene3DTransformTrack* TransformTrack = NewMovieScene->AddTrack<UMovieScene3DTransformTrack>(PossessableGuid);								
-					FMatineeImportTools::CopyInterpMoveTrack(MatineeMoveTrack, TransformTrack);
+					FMatineeImportTools::CopyInterpMoveTrack(MatineeMoveTrack, TransformTrack, DefaultScale);
 				}
 			}
 			else if (Track->IsA(UInterpTrackAnimControl::StaticClass()))
@@ -549,6 +552,19 @@ protected:
 					}
 				}
 			}
+			else if (Track->IsA(UInterpTrackVectorProp::StaticClass()))
+			{
+				UInterpTrackVectorProp* MatineeVectorTrack = StaticCast<UInterpTrackVectorProp*>(Track);
+				if (MatineeVectorTrack->GetNumKeyframes() != 0 && GroupActor && PossessableGuid.IsValid())
+				{
+					UMovieSceneVectorTrack* VectorTrack = AddPropertyTrack<UMovieSceneVectorTrack>(MatineeVectorTrack->PropertyName, GroupActor, PossessableGuid, Player, NewSequence, NewMovieScene, NumWarnings);
+					if (VectorTrack)
+					{
+						VectorTrack->SetNumChannelsUsed(3);
+						FMatineeImportTools::CopyInterpVectorTrack(MatineeVectorTrack, VectorTrack);
+					}
+				}
+			}
 			else if (Track->IsA(UInterpTrackColorProp::StaticClass()))
 			{
 				UInterpTrackColorProp* MatineeColorTrack = StaticCast<UInterpTrackColorProp*>(Track);
diff --git a/Engine/Shaders/MediaShaders.usf b/Engine/Shaders/MediaShaders.usf
--- a/Engine/Shaders/MediaShaders.usf
+++ b/Engine/Shaders/MediaShaders.usf
@@ -8,53 +8,26 @@
 #include "GammaCorrectionCommon.usf"
 
 
-/* Constants
+/* Helpers
 =============================================================================*/
 
-static const float3x3 YUV_TO_RGB_BT601 =
-{
-	1.000000, 0.000000, 1.139830,
-	1.000000,-0.394650,-0.580600,
-	1.000000, 2.032110, 0.000000,
-};
-
-static const float3x3 YUV_TO_RGB_BT709 =
-{
-	1.000000, 0.000000, 1.280330,
-	1.000000,-0.214820,-0.380590,
-	1.000000, 2.127980, 0.000000,
-};
-
-static const float3x3 YUV_TO_RGB_JPEG =
-{
-	1.000000, 0.000000, 1.402000,
-	1.000000,-0.344140,-0.714140,
-	1.000000, 1.772000, 0.000000,
-};
-
-static const float3x3 YUV_TO_RGB_MSDN =
+float3 YuvToRgb(float3 YUV, float4x4 ColorTransform, uint SrgbToLinear)
 {
-	1.164383, 0.000000, 1.596027,
-	1.164383,-0.391762,-0.812968,
-	1.164383, 2.017232, 0.000000,
-};
-
-static const float3x3 YUV_TO_RGB_PS4 =
-{
-	1.164400, 0.000000, 1.792700,
-	1.164400,-0.213300,-0.532900,
-	1.164400, 2.112400, 0.000000,
-};
+	float3x3 Coeff = float3x3(
+		ColorTransform[0].xyz,
+		ColorTransform[1].xyz,
+		ColorTransform[2].xyz
+	);
 
+	YUV -= float3(0.0625, 0.5, 0.5);
+	float3 sRGB = mul(Coeff, YUV);
 
-/* Helpers
-=============================================================================*/
+	if (SrgbToLinear == 0)
+	{
+		return sRGB;
+	}
 
-float3 YuvToRgb(float3 YUV, float3x3 Coeff)
-{
-	YUV -= float3(0.0625, 0.5, 0.5);
-	float3 RGB = mul(Coeff, YUV);
-	return sRGBToLinear(RGB);
+	return sRGBToLinear(sRGB);
 }
 
 
@@ -81,9 +54,9 @@ void AYUVConvertPS(
 	float2 InUV : TEXCOORD0,
 	out float4 OutColor : SV_Target0)
 {
-	float4 AYUV = AYUVConvertUB.AYUVTexture.Sample(AYUVConvertUB.AYUVSampler, InUV.xy);
+	float4 AYUV = AYUVConvertUB.Texture.Sample(AYUVConvertUB.Sampler, InUV.xy);
 	
-	OutColor.xyz = YuvToRgb(AYUV.zyx, YUV_TO_RGB_MSDN);
+	OutColor.xyz = YuvToRgb(AYUV.zyx, AYUVConvertUB.ColorTransform, AYUVConvertUB.SrgbToLinear);
 	OutColor.w = AYUV.w;
 }
 
@@ -96,7 +69,7 @@ void BMPConvertPS(
 {
 	InUV.y = 1.0f - InUV.y;
 	float2 ScaledUV = InUV * BMPConvertUB.UVScale;
-	OutColor = BMPConvertUB.BMPTexture.Sample(BMPConvertUB.BMPSampler, ScaledUV);
+	OutColor = BMPConvertUB.Texture.Sample(BMPConvertUB.Sampler, ScaledUV);
 }
 
 
@@ -109,7 +82,7 @@ void NV12ConvertPS(
 	float2 ScaledUV = InUV * NV12ConvertUB.UVScale;
 
 	float3 YUV;
-	YUV.x = NV12ConvertUB.NV12Texture.Sample(NV12ConvertUB.NV12SamplerB, ScaledUV).x;
+	YUV.x = NV12ConvertUB.Texture.Sample(NV12ConvertUB.SamplerB, ScaledUV).x;
 
 	ScaledUV.y *= 0.5f;
 	ScaledUV.y += 2.0f / 3.0f;
@@ -118,10 +91,10 @@ void NV12ConvertPS(
 	float UPos = (float(UVOffs * NV12ConvertUB.UVScale.x) + 0.5f) / NV12ConvertUB.OutputWidth;
 	float VPos = (float(UVOffs * NV12ConvertUB.UVScale.x) + 1.5f) / NV12ConvertUB.OutputWidth;
 
-	YUV.y = NV12ConvertUB.NV12Texture.Sample(NV12ConvertUB.NV12SamplerP, float2(UPos, ScaledUV.y)).x;
-	YUV.z = NV12ConvertUB.NV12Texture.Sample(NV12ConvertUB.NV12SamplerP, float2(VPos, ScaledUV.y)).x;
+	YUV.y = NV12ConvertUB.Texture.Sample(NV12ConvertUB.SamplerP, float2(UPos, ScaledUV.y)).x;
+	YUV.z = NV12ConvertUB.Texture.Sample(NV12ConvertUB.SamplerP, float2(VPos, ScaledUV.y)).x;
 
-	OutColor.xyz = YuvToRgb(YUV, YUV_TO_RGB_MSDN);
+	OutColor.xyz = YuvToRgb(YUV, NV12ConvertUB.ColorTransform, NV12ConvertUB.SrgbToLinear);
 	OutColor.w = 1.0f;
 }
 
@@ -135,7 +108,7 @@ void NV21ConvertPS(
 	float2 ScaledUV = InUV * NV21ConvertUB.UVScale;
 
 	float3 YUV;
-	YUV.x = NV21ConvertUB.NV21Texture.Sample(NV21ConvertUB.NV21SamplerB, ScaledUV).x;
+	YUV.x = NV21ConvertUB.Texture.Sample(NV21ConvertUB.SamplerB, ScaledUV).x;
 
 	ScaledUV.y *= 0.5f;
 	ScaledUV.y += 2.0f / 3.0f;
@@ -144,10 +117,10 @@ void NV21ConvertPS(
 	float UPos = (float(UVOffs * NV21ConvertUB.UVScale.x) + 1.5f) / NV21ConvertUB.OutputWidth;
 	float VPos = (float(UVOffs * NV21ConvertUB.UVScale.x) + 0.5f) / NV21ConvertUB.OutputWidth;
 
-	YUV.y = NV21ConvertUB.NV21Texture.Sample(NV21ConvertUB.NV21SamplerP, float2(UPos, ScaledUV.y)).x;
-	YUV.z = NV21ConvertUB.NV21Texture.Sample(NV21ConvertUB.NV21SamplerP, float2(VPos, ScaledUV.y)).x;
+	YUV.y = NV21ConvertUB.Texture.Sample(NV21ConvertUB.SamplerP, float2(UPos, ScaledUV.y)).x;
+	YUV.z = NV21ConvertUB.Texture.Sample(NV21ConvertUB.SamplerP, float2(VPos, ScaledUV.y)).x;
 
-	OutColor.xyz = YuvToRgb(YUV, YUV_TO_RGB_MSDN);
+	OutColor.xyz = YuvToRgb(YUV, NV21ConvertUB.ColorTransform, NV21ConvertUB.SrgbToLinear);
 	OutColor.w = 1.0f;
 }
 
@@ -159,8 +132,9 @@ void RGBConvertPS(
 	out float4 OutColor : SV_Target0)
 {
 	float2 ScaledUV = InUV * RGBConvertUB.UVScale;
-	float4 RGB = RGBConvertUB.RGBTexture.Sample(RGBConvertUB.RGBSampler, ScaledUV);
-	OutColor.xyz = sRGBToLinear(RGB);
+	float4 RGB = RGBConvertUB.Texture.Sample(RGBConvertUB.Sampler, ScaledUV);
+
+	OutColor.xyz = (RGBConvertUB.SrgbToLinear == 1) ? sRGBToLinear(RGB) : RGB;
 }
 
 
@@ -174,7 +148,7 @@ void YCbCrConvertPS(
 	YUV.x = YCbCrConvertUB.LumaTexture.Sample(YCbCrConvertUB.LumaSampler, InUV.xy).x;
 	YUV.yz = YCbCrConvertUB.CbCrTexture.Sample(YCbCrConvertUB.CbCrSampler, InUV.xy).xy;
 	
-	OutColor.xyz = YuvToRgb(YUV, YUV_TO_RGB_PS4);
+	OutColor.xyz = YuvToRgb(YUV, YCbCrConvertUB.ColorTransform, YCbCrConvertUB.SrgbToLinear);
 	OutColor.w = 1.0f;
 }
 
@@ -185,8 +159,8 @@ void UYVYConvertPS(
 	float2 InUV : TEXCOORD0,
 	out float4 OutColor : SV_Target0)
 {
-	float4 UYVYB = UYVYConvertUB.UYVYTexture.Sample(UYVYConvertUB.UYVYSamplerB, float2(InUV.x, InUV.y));
-	float4 UYVYP = UYVYConvertUB.UYVYTexture.Sample(UYVYConvertUB.UYVYSamplerP, float2(InUV.x, InUV.y));
+	float4 UYVYB = UYVYConvertUB.Texture.Sample(UYVYConvertUB.SamplerB, float2(InUV.x, InUV.y));
+	float4 UYVYP = UYVYConvertUB.Texture.Sample(UYVYConvertUB.SamplerP, float2(InUV.x, InUV.y));
 
 	float Pos = 2.0f * InUV.x * UYVYConvertUB.Width;
 	float3 YUV;
@@ -194,7 +168,7 @@ void UYVYConvertPS(
 	YUV.x = (floor(Pos % 2.0f) == 0) ? UYVYP.y : UYVYP.w;
 	YUV.yz = UYVYB.zx;
 
-	OutColor.xyz = YuvToRgb(YUV, YUV_TO_RGB_MSDN);
+	OutColor.xyz = YuvToRgb(YUV, UYVYConvertUB.ColorTransform, UYVYConvertUB.SrgbToLinear);
 	OutColor.w = 1.0f;
 }
 
@@ -210,7 +184,7 @@ void YUVConvertPS(
 	YUV.y = YUVConvertUB.UTexture.Sample(YUVConvertUB.USampler, InUV.xy).x;
 	YUV.z = YUVConvertUB.VTexture.Sample(YUVConvertUB.VSampler, InUV.xy).x;
 	
-	OutColor.xyz = YuvToRgb(YUV, YUV_TO_RGB_MSDN);
+	OutColor.xyz = YuvToRgb(YUV, YUVConvertUB.ColorTransform, YUVConvertUB.SrgbToLinear);
 	OutColor.w = 1.0f;
 }
 
@@ -224,13 +198,13 @@ void YUY2ConvertPS(
 	float2 ScaledUV = InUV * YUY2ConvertUB.UVScale;
 
 	float3 YUV;
-	YUV.yz = YUY2ConvertUB.YUY2Texture.Sample(YUY2ConvertUB.YUY2SamplerB, ScaledUV).yw;
+	YUV.yz = YUY2ConvertUB.Texture.Sample(YUY2ConvertUB.SamplerB, ScaledUV).yw;
 
-	float4 YUY2P = YUY2ConvertUB.YUY2Texture.Sample(YUY2ConvertUB.YUY2SamplerP, ScaledUV);
+	float4 YUY2P = YUY2ConvertUB.Texture.Sample(YUY2ConvertUB.SamplerP, ScaledUV);
 	float Pos = InUV.x * YUY2ConvertUB.OutputWidth;
 	YUV.x = (floor(Pos % 2.0f) == 0) ? YUY2P.z : YUY2P.x;
 
-	OutColor.xyz = YuvToRgb(YUV, YUV_TO_RGB_MSDN);
+	OutColor.xyz = YuvToRgb(YUV, YUY2ConvertUB.ColorTransform, YUY2ConvertUB.SrgbToLinear);
 	OutColor.w = 1.0f;
 }
 
@@ -241,8 +215,8 @@ void YVYUConvertPS(
 	float2 InUV : TEXCOORD0,
 	out float4 OutColor : SV_Target0)
 {
-	float4 YVYUB = YVYUConvertUB.YVYUTexture.Sample(YVYUConvertUB.YVYUSamplerB, float2(InUV.x, InUV.y));
-	float4 YVYUP = YVYUConvertUB.YVYUTexture.Sample(YVYUConvertUB.YVYUSamplerP, float2(InUV.x, InUV.y));
+	float4 YVYUB = YVYUConvertUB.Texture.Sample(YVYUConvertUB.SamplerB, float2(InUV.x, InUV.y));
+	float4 YVYUP = YVYUConvertUB.Texture.Sample(YVYUConvertUB.SamplerP, float2(InUV.x, InUV.y));
 
 	float Pos = 2.0f * InUV.x * YVYUConvertUB.Width;
 	float3 YUV;
@@ -250,6 +224,6 @@ void YVYUConvertPS(
 	YUV.x = (floor(Pos % 2.0f) == 0) ? YVYUP.z : YVYUP.x;
 	YUV.yz = YVYUB.wy;
 
-	OutColor.xyz = YuvToRgb(YUV, YUV_TO_RGB_MSDN);
+	OutColor.xyz = YuvToRgb(YUV, YVYUConvertUB.ColorTransform, YVYUConvertUB.SrgbToLinear);
 	OutColor.w = 1.0f;
 }
diff --git a/Engine/Source/Developer/AutomationDriver/Private/Specs/AutomationDriverSpecSuiteViewModel.cpp b/Engine/Source/Developer/AutomationDriver/Private/Specs/AutomationDriverSpecSuiteViewModel.cpp
--- a/Engine/Source/Developer/AutomationDriver/Private/Specs/AutomationDriverSpecSuiteViewModel.cpp
+++ b/Engine/Source/Developer/AutomationDriver/Private/Specs/AutomationDriverSpecSuiteViewModel.cpp
@@ -43,7 +43,7 @@ public:
 
 	virtual bool IsKeyEnabled(EPianoKey Key) const override
 	{
-		if (KeyResetDelay == FTimespan::Zero())
+		if (KeyResetDelay.IsZero())
 		{
 			return true;
 		}
diff --git a/Engine/Source/Developer/DesktopPlatform/Private/PlatformInfo.cpp b/Engine/Source/Developer/DesktopPlatform/Private/PlatformInfo.cpp
--- a/Engine/Source/Developer/DesktopPlatform/Private/PlatformInfo.cpp
+++ b/Engine/Source/Developer/DesktopPlatform/Private/PlatformInfo.cpp
@@ -11,7 +11,7 @@ namespace PlatformInfo
 namespace
 {
 
-FPlatformInfo BuildPlatformInfo(const FName& InPlatformInfoName, const FName& InTargetPlatformName, const FText& InDisplayName, const EPlatformType InPlatformType, const EPlatformFlags::Flags InPlatformFlags, const FPlatformIconPaths& InIconPaths, const FString& InUATCommandLine, const FString& InAutoSDKPath, EPlatformSDKStatus InStatus, const FString& InTutorial, bool InEnabled, FString InBinaryFolderName, bool InUsesHostCompiler, bool InUATClosesAfterLaunch)
+FPlatformInfo BuildPlatformInfo(const FName& InPlatformInfoName, const FName& InTargetPlatformName, const FText& InDisplayName, const EPlatformType InPlatformType, const EPlatformFlags::Flags InPlatformFlags, const FPlatformIconPaths& InIconPaths, const FString& InUATCommandLine, const FString& InAutoSDKPath, EPlatformSDKStatus InStatus, const FString& InTutorial, bool InEnabled, FString InBinaryFolderName, FString InIniPlatformName, bool InUsesHostCompiler, bool InUATClosesAfterLaunch)
 {
 	FPlatformInfo PlatformInfo;
 
@@ -40,6 +40,7 @@ FPlatformInfo BuildPlatformInfo(const FName& InPlatformInfoName, const FName& In
 	PlatformInfo.UATCommandLine = InUATCommandLine;
 	PlatformInfo.AutoSDKPath = InAutoSDKPath;
 	PlatformInfo.BinaryFolderName = InBinaryFolderName;
+	PlatformInfo.IniPlatformName = InIniPlatformName;
 
 	// Generate the icon style names for FEditorStyle
 	PlatformInfo.IconPaths.NormalStyleName = *FString::Printf(TEXT("Launcher.Platform_%s"), *InPlatformInfoNameString);
@@ -76,46 +77,48 @@ static const FString IOSTutorial = TEXT("/Engine/Tutorial/NotYetImplemented");
 #endif
 
 static const FPlatformInfo PlatformInfoArray[] = {
-	//				  PlatformInfoName					TargetPlatformName			DisplayName													PlatformType				PlatformFlags					IconPaths																																			UATCommandLine											AutoSDKPath			SDKStatus						SDKTutorial																								EnabledForUse								BinaryFolderName	UsesHostCompiler		UATClosesAfterLaunch
-	BuildPlatformInfo(TEXT("WindowsNoEditor"),			TEXT("WindowsNoEditor"),	LOCTEXT("WindowsNoEditor", "Windows"),						EPlatformType::Game,		EPlatformFlags::None,			FPlatformIconPaths(TEXT("Launcher/WindowsTarget/Platform_WindowsNoEditor_24x"), TEXT("Launcher/WindowsTarget/Platform_WindowsNoEditor_128x")),		TEXT("-targetplatform=Win64"),							TEXT(""),			EPlatformSDKStatus::Unknown,	TEXT("/Engine/Tutorial/Installation/InstallingVisualStudioTutorial.InstallingVisualStudioTutorial"),	IsAvailableOnWindows,						TEXT("Win64"),		IsAvailableOnWindows,	false),
-	BuildPlatformInfo(TEXT("WindowsNoEditor_Win32"),	TEXT("WindowsNoEditor"),	LOCTEXT("WindowsNoEditor_Win32", "Windows (32-bit)"),		EPlatformType::Game,		EPlatformFlags::BuildFlavor,	FPlatformIconPaths(TEXT("Launcher/WindowsTarget/Platform_WindowsNoEditor_24x"), TEXT("Launcher/WindowsTarget/Platform_WindowsNoEditor_128x")),		TEXT("-targetplatform=Win32"),							TEXT(""),			EPlatformSDKStatus::Unknown,	TEXT("/Engine/Tutorial/Installation/InstallingVisualStudioTutorial.InstallingVisualStudioTutorial"),	IsAvailableOnWindows,						TEXT("Win32"),		IsAvailableOnWindows,	false),
-	BuildPlatformInfo(TEXT("WindowsNoEditor_Win64"),	TEXT("WindowsNoEditor"),	LOCTEXT("WindowsNoEditor_Win64", "Windows (64-bit)"),		EPlatformType::Game,		EPlatformFlags::BuildFlavor,	FPlatformIconPaths(TEXT("Launcher/WindowsTarget/Platform_WindowsNoEditor_24x"), TEXT("Launcher/WindowsTarget/Platform_WindowsNoEditor_128x")),		TEXT("-targetplatform=Win64"),							TEXT(""),			EPlatformSDKStatus::Unknown,	TEXT("/Engine/Tutorial/Installation/InstallingVisualStudioTutorial.InstallingVisualStudioTutorial"),	IsAvailableOnWindows,						TEXT("Win64"),		IsAvailableOnWindows,	false),
-	BuildPlatformInfo(TEXT("Windows"),					TEXT("Windows"),			LOCTEXT("WindowsEditor", "Windows (Editor)"),				EPlatformType::Editor,		EPlatformFlags::None,			FPlatformIconPaths(TEXT("Launcher/WindowsTarget/Platform_Windows_24x"), TEXT("Launcher/WindowsTarget/Platform_Windows_128x")),						TEXT(""),												TEXT(""),			EPlatformSDKStatus::Unknown,	TEXT("/Engine/Tutorial/Installation/InstallingVisualStudioTutorial.InstallingVisualStudioTutorial"),	IsAvailableOnWindows,						TEXT("Win64"),		IsAvailableOnWindows,	false),
-	BuildPlatformInfo(TEXT("WindowsClient"),			TEXT("WindowsClient"),		LOCTEXT("WindowsClient", "Windows (Client-only)"),			EPlatformType::Client,		EPlatformFlags::None,			FPlatformIconPaths(TEXT("Launcher/WindowsTarget/Platform_Windows_24x"), TEXT("Launcher/WindowsTarget/Platform_Windows_128x")),						TEXT("-client"),										TEXT(""),			EPlatformSDKStatus::Unknown,	TEXT("/Engine/Tutorial/Installation/InstallingVisualStudioTutorial.InstallingVisualStudioTutorial"),	IsAvailableOnWindows,						TEXT("Win64"),		IsAvailableOnWindows,	false),
-	BuildPlatformInfo(TEXT("WindowsServer"),			TEXT("WindowsServer"),		LOCTEXT("WindowsServer", "Windows (Dedicated Server)"),		EPlatformType::Server,		EPlatformFlags::None,			FPlatformIconPaths(TEXT("Launcher/WindowsTarget/Platform_WindowsServer_24x"), TEXT("Launcher/WindowsTarget/Platform_WindowsServer_128x")),			TEXT(""),												TEXT(""),			EPlatformSDKStatus::Unknown,	TEXT("/Engine/Tutorial/Installation/InstallingVisualStudioTutorial.InstallingVisualStudioTutorial"),	IsAvailableOnWindows,						TEXT("Win64"),		IsAvailableOnWindows,	false),
-
-	BuildPlatformInfo(TEXT("MacNoEditor"),				TEXT("MacNoEditor"),		LOCTEXT("MacNoEditor", "Mac"),								EPlatformType::Game,		EPlatformFlags::None,			FPlatformIconPaths(TEXT("Launcher/MacTarget/Platform_Mac_24x"), TEXT("Launcher/MacTarget/Platform_Mac_128x")),										TEXT("-targetplatform=Mac"),							TEXT(""),			EPlatformSDKStatus::Unknown,	TEXT("/Engine/Tutorial/Installation/InstallingXCodeTutorial.InstallingXCodeTutorial"),					IsAvailableOnMac,							TEXT("Mac"),		IsAvailableOnMac,		true),
-	BuildPlatformInfo(TEXT("Mac"),						TEXT("Mac"),				LOCTEXT("MacEditor", "Mac (Editor)"),						EPlatformType::Editor,		EPlatformFlags::None,			FPlatformIconPaths(TEXT("Launcher/MacTarget/Platform_Mac_24x"), TEXT("Launcher/MacTarget/Platform_Mac_128x")),										TEXT(""),												TEXT(""),			EPlatformSDKStatus::Unknown,	TEXT("/Engine/Tutorial/Installation/InstallingXCodeTutorial.InstallingXCodeTutorial"),					IsAvailableOnMac,							TEXT("Mac"),		IsAvailableOnMac,		true),
-	BuildPlatformInfo(TEXT("MacClient"),				TEXT("MacClient"),			LOCTEXT("MacClient", "Mac (Client-only)"),					EPlatformType::Client,		EPlatformFlags::None,			FPlatformIconPaths(TEXT("Launcher/MacTarget/Platform_Mac_24x"), TEXT("Launcher/MacTarget/Platform_Mac_128x")),										TEXT("-client"),										TEXT(""),			EPlatformSDKStatus::Unknown,	TEXT("/Engine/Tutorial/Installation/InstallingXCodeTutorial.InstallingXCodeTutorial"),					IsAvailableOnMac,							TEXT("Mac"),		IsAvailableOnMac,		true),
-	BuildPlatformInfo(TEXT("MacServer"),				TEXT("MacServer"),			LOCTEXT("MacServer", "Mac (Dedicated Server)"),				EPlatformType::Server,		EPlatformFlags::None,			FPlatformIconPaths(TEXT("Launcher/MacTarget/Platform_Mac_24x"), TEXT("Launcher/MacTarget/Platform_Mac_128x")),										TEXT(""),												TEXT(""),			EPlatformSDKStatus::Unknown,	TEXT("/Engine/Tutorial/Installation/InstallingXCodeTutorial.InstallingXCodeTutorial"),					IsAvailableOnMac,							TEXT("Mac"),		IsAvailableOnMac,		true),
-
-	BuildPlatformInfo(TEXT("LinuxNoEditor"),			TEXT("LinuxNoEditor"),		LOCTEXT("LinuxNoEditor", "Linux"),							EPlatformType::Game,		EPlatformFlags::None,			FPlatformIconPaths(TEXT("Launcher/LinuxTarget/Platform_Linux_24x"), TEXT("Launcher/LinuxTarget/Platform_Linux_128x")),								TEXT("-targetplatform=Linux"),							TEXT("Linux_x64"),	EPlatformSDKStatus::Unknown,	TEXT("/Platforms/Linux/GettingStarted"),																IsAvailableOnLinux || IsAvailableOnWindows,	TEXT("Linux"),		IsAvailableOnLinux,		true),	
-	BuildPlatformInfo(TEXT("Linux"),					TEXT("Linux"),				LOCTEXT("LinuxEditor", "Linux (Editor)"),					EPlatformType::Editor,		EPlatformFlags::None,			FPlatformIconPaths(TEXT("Launcher/LinuxTarget/Platform_Linux_24x"), TEXT("Launcher/LinuxTarget/Platform_Linux_128x")),								TEXT(""),												TEXT("Linux_x64"),	EPlatformSDKStatus::Unknown,	TEXT("/Platforms/Linux/GettingStarted"),																IsAvailableOnLinux,							TEXT("Linux"),		IsAvailableOnLinux,		true),
-	BuildPlatformInfo(TEXT("LinuxClient"),				TEXT("LinuxClient"),		LOCTEXT("LinuxClient", "Linux (Client-only)"),				EPlatformType::Client,		EPlatformFlags::None,			FPlatformIconPaths(TEXT("Launcher/LinuxTarget/Platform_Linux_24x"), TEXT("Launcher/LinuxTarget/Platform_Linux_128x")),								TEXT("-client"),										TEXT("Linux_x64"),	EPlatformSDKStatus::Unknown,	TEXT("/Platforms/Linux/GettingStarted"),																IsAvailableOnLinux || IsAvailableOnWindows,	TEXT("Linux"),		IsAvailableOnLinux,		true),
-	BuildPlatformInfo(TEXT("LinuxServer"),				TEXT("LinuxServer"),		LOCTEXT("LinuxServer", "Linux (Dedicated Server)"),			EPlatformType::Server,		EPlatformFlags::None,			FPlatformIconPaths(TEXT("Launcher/LinuxTarget/Platform_Linux_24x"), TEXT("Launcher/LinuxTarget/Platform_Linux_128x")),								TEXT(""),												TEXT("Linux_x64"),	EPlatformSDKStatus::Unknown,	TEXT("/Platforms/Linux/GettingStarted"),																IsAvailableOnLinux || IsAvailableOnWindows,	TEXT("Linux"),		IsAvailableOnLinux,		true),
-
-	BuildPlatformInfo(TEXT("IOS"),						TEXT("IOS"),				LOCTEXT("IOS", "iOS"),										EPlatformType::Game,		EPlatformFlags::None,			FPlatformIconPaths(TEXT("Launcher/iOSTarget/Platform_iOS_24x"), TEXT("Launcher/iOSTarget/Platform_iOS_128x")),										TEXT(""),												TEXT(""),			EPlatformSDKStatus::Unknown,	IOSTutorial,																							IsAvailableOnWindows || IsAvailableOnMac,	TEXT("IOS"),		false,					true),
-
-	BuildPlatformInfo(TEXT("Android"),					TEXT("Android"),			LOCTEXT("Android", "Android"),								EPlatformType::Game,		EPlatformFlags::None,			FPlatformIconPaths(TEXT("Launcher/AndroidTarget/Platform_Android_24x"), TEXT("Launcher/AndroidTarget/Platform_Android_128x")),						TEXT(""),												TEXT("Android"),	EPlatformSDKStatus::Unknown,	TEXT("/Engine/Tutorial/Mobile/SettingUpAndroidTutorial.SettingUpAndroidTutorial"),						IsAvailableOnLinux || IsAvailableOnWindows || IsAvailableOnMac,	TEXT("Android"),	false,					true),
-	BuildPlatformInfo(TEXT("Android_All"),				TEXT("Android"),			LOCTEXT("Android_All", "Android (All)"),					EPlatformType::Game,		EPlatformFlags::CookFlavor,		FPlatformIconPaths(TEXT("Launcher/AndroidTarget/Platform_Android_24x"), TEXT("Launcher/AndroidTarget/Platform_Android_128x")),						TEXT(""),												TEXT("Android"),	EPlatformSDKStatus::Unknown,	TEXT("/Engine/Tutorial/Mobile/SettingUpAndroidTutorial.SettingUpAndroidTutorial"),						IsAvailableOnLinux || IsAvailableOnWindows || IsAvailableOnMac,	TEXT("Android"),	false,					true),
-	BuildPlatformInfo(TEXT("Android_Multi"),			TEXT("Android_Multi"),		LOCTEXT("Android_Multi", "Android (Multi)"),				EPlatformType::Game,		EPlatformFlags::CookFlavor,		FPlatformIconPaths(TEXT("Launcher/AndroidTarget/Platform_Android_24x"), TEXT("Launcher/AndroidTarget/Platform_Android_128x")),						TEXT("-targetplatform=Android -cookflavor=Multi"),		TEXT("Android"),	EPlatformSDKStatus::Unknown,	TEXT("/Engine/Tutorial/Mobile/SettingUpAndroidTutorial.SettingUpAndroidTutorial"),						IsAvailableOnLinux || IsAvailableOnWindows || IsAvailableOnMac,	TEXT("Android"),	false,					true),
-	BuildPlatformInfo(TEXT("Android_ATC"),				TEXT("Android_ATC"),		LOCTEXT("Android_ATC", "Android (ATC)"),					EPlatformType::Game,		EPlatformFlags::CookFlavor,		FPlatformIconPaths(TEXT("Launcher/AndroidTarget/Platform_Android_ATC_24x"), TEXT("Launcher/AndroidTarget/Platform_Android_128x")),					TEXT("-targetplatform=Android -cookflavor=ATC"),		TEXT("Android"),	EPlatformSDKStatus::Unknown,	TEXT("/Engine/Tutorial/Mobile/SettingUpAndroidTutorial.SettingUpAndroidTutorial"),						IsAvailableOnLinux || IsAvailableOnWindows || IsAvailableOnMac,	TEXT("Android"),	false,					true),
-	BuildPlatformInfo(TEXT("Android_DXT"),				TEXT("Android_DXT"),		LOCTEXT("Android_DXT", "Android (DXT)"),					EPlatformType::Game,		EPlatformFlags::CookFlavor,		FPlatformIconPaths(TEXT("Launcher/AndroidTarget/Platform_Android_DXT_24x"), TEXT("Launcher/AndroidTarget/Platform_Android_128x")),					TEXT("-targetplatform=Android -cookflavor=DXT"),		TEXT("Android"),	EPlatformSDKStatus::Unknown,	TEXT("/Engine/Tutorial/Mobile/SettingUpAndroidTutorial.SettingUpAndroidTutorial"),						IsAvailableOnLinux || IsAvailableOnWindows || IsAvailableOnMac,	TEXT("Android"),	false,					true),
-	BuildPlatformInfo(TEXT("Android_ETC1"),				TEXT("Android_ETC1"),		LOCTEXT("Android_ETC1", "Android (ETC1)"),					EPlatformType::Game,		EPlatformFlags::CookFlavor,		FPlatformIconPaths(TEXT("Launcher/AndroidTarget/Platform_Android_ETC1_24x"), TEXT("Launcher/AndroidTarget/Platform_Android_128x")),					TEXT("-targetplatform=Android -cookflavor=ETC1"),		TEXT("Android"),	EPlatformSDKStatus::Unknown,	TEXT("/Engine/Tutorial/Mobile/SettingUpAndroidTutorial.SettingUpAndroidTutorial"),						IsAvailableOnLinux || IsAvailableOnWindows || IsAvailableOnMac,	TEXT("Android"),	false,					true),
-	BuildPlatformInfo(TEXT("Android_ETC2"),				TEXT("Android_ETC2"),		LOCTEXT("Android_ETC2", "Android (ETC2)"),					EPlatformType::Game,		EPlatformFlags::CookFlavor,		FPlatformIconPaths(TEXT("Launcher/AndroidTarget/Platform_Android_ETC2_24x"), TEXT("Launcher/AndroidTarget/Platform_Android_128x")),					TEXT("-targetplatform=Android -cookflavor=ETC2"),		TEXT("Android"),	EPlatformSDKStatus::Unknown,	TEXT("/Engine/Tutorial/Mobile/SettingUpAndroidTutorial.SettingUpAndroidTutorial"),						IsAvailableOnLinux || IsAvailableOnWindows || IsAvailableOnMac,	TEXT("Android"),	false,					true),
-	BuildPlatformInfo(TEXT("Android_PVRTC"),			TEXT("Android_PVRTC"),		LOCTEXT("Android_PVRTC", "Android (PVRTC)"),				EPlatformType::Game,		EPlatformFlags::CookFlavor,		FPlatformIconPaths(TEXT("Launcher/AndroidTarget/Platform_Android_PVRTC_24x"), TEXT("Launcher/AndroidTarget/Platform_Android_128x")),				TEXT("-targetplatform=Android -cookflavor=PVRTC"),		TEXT("Android"),	EPlatformSDKStatus::Unknown,	TEXT("/Engine/Tutorial/Mobile/SettingUpAndroidTutorial.SettingUpAndroidTutorial"),						IsAvailableOnLinux || IsAvailableOnWindows || IsAvailableOnMac,	TEXT("Android"),	false,					true),
-	BuildPlatformInfo(TEXT("Android_ASTC"),				TEXT("Android_ASTC"),		LOCTEXT("Android_ASTC", "Android (ASTC)"),					EPlatformType::Game,		EPlatformFlags::CookFlavor,		FPlatformIconPaths(TEXT("Launcher/AndroidTarget/Platform_Android_ASTC_24x"), TEXT("Launcher/AndroidTarget/Platform_Android_128x")),					TEXT("-targetplatform=Android -cookflavor=ASTC"),		TEXT("Android"),	EPlatformSDKStatus::Unknown,	TEXT("/Engine/Tutorial/Mobile/SettingUpAndroidTutorial.SettingUpAndroidTutorial"),						IsAvailableOnLinux || IsAvailableOnWindows || IsAvailableOnMac,	TEXT("Android"),	false,					true),
-
-	BuildPlatformInfo(TEXT("HTML5"),					TEXT("HTML5"),				LOCTEXT("HTML5", "HTML5"),									EPlatformType::Game,		EPlatformFlags::None,			FPlatformIconPaths(TEXT("Launcher/HTML5Target/Platform_HTML5_24x"), TEXT("Launcher/HTML5Target/Platform_HTML5_128x")),								TEXT(""),												TEXT("HTML5"),		EPlatformSDKStatus::Unknown,	TEXT("/Platforms/HTML5/GettingStarted"),																IsAvailableOnWindows || IsAvailableOnMac,   TEXT("HTML5"),		false,					true),
-
-	BuildPlatformInfo(TEXT("PS4"),						TEXT("PS4"),				LOCTEXT("PS4", "PlayStation 4"),							EPlatformType::Game,		EPlatformFlags::None,			FPlatformIconPaths(TEXT("Launcher/PS4Target/Platform_PS4_24x"), TEXT("Launcher/PS4Target/Platform_PS4_128x")),										TEXT(""),												TEXT("PS4"),		EPlatformSDKStatus::Unknown,	TEXT("/Platforms/PS4/GettingStarted"),																	IsAvailableOnWindows,						TEXT("PS4"),		false,					false),
-
-	BuildPlatformInfo(TEXT("XboxOne"),					TEXT("XboxOne"),			LOCTEXT("XboxOne", "Xbox One"),								EPlatformType::Game,		EPlatformFlags::None,			FPlatformIconPaths(TEXT("Launcher/XboxOneTarget/Platform_XboxOne_24x"), TEXT("Launcher/XboxOneTarget/Platform_XboxOne_128x")),						TEXT(""),												TEXT("XboxOne"),	EPlatformSDKStatus::Unknown,	TEXT("/Platforms/XboxOne/GettingStarted"),																IsAvailableOnWindows,						TEXT("XboxOne"),		false,					true),
-
-	BuildPlatformInfo(TEXT("AllDesktop"),					TEXT("AllDesktop"),			LOCTEXT("DesktopTargetPlatDisplay", "Desktop (Win+Mac+Linux)"),	EPlatformType::Game,	EPlatformFlags::None,			FPlatformIconPaths(TEXT("Launcher/DesktopTarget/Platform_Desktop_24x"), TEXT("Launcher/DesktopTarget/Platform_Desktop_128x")),						TEXT(""),												TEXT(""),			EPlatformSDKStatus::Unknown,	TEXT(""),																							IsAvailableOnWindows/*  see SProjectTargetPlatformSettings::Construct !!!! IsAvailableOnWindows || IsAvailableOnMac || IsAvailableOnLinux*/,	TEXT(""), false, true),
-
-	BuildPlatformInfo(TEXT("TVOS"),					TEXT("TVOS"),			LOCTEXT("TVOSTargetPlatDisplay", "tvOS"),	EPlatformType::Game,	EPlatformFlags::None,			FPlatformIconPaths(TEXT("Launcher/TVOSTarget/Platform_TVOS_24x"), TEXT("Launcher/TVOSTarget/Platform_TVOS_128x")),						TEXT(""),												TEXT(""),			EPlatformSDKStatus::Unknown,	TEXT(""),																							IsAvailableOnWindows || IsAvailableOnMac,	TEXT("TVOS"), false, true),
-	BuildPlatformInfo(TEXT("Switch"),						TEXT("Switch"),				LOCTEXT("Switch", "Switch"),									EPlatformType::Game,		EPlatformFlags::None,			FPlatformIconPaths(TEXT("Launcher/SwitchTarget/Platform_Switch_24x"), TEXT("Launcher/SwitchTarget/Platform_Switch_128x")),						TEXT(""),												TEXT(""),			EPlatformSDKStatus::Unknown,	TEXT(""),																								IsAvailableOnWindows,						TEXT("Switch"),		false,					true),
+	// PlatformInfoName									TargetPlatformName			DisplayName														PlatformType			PlatformFlags					IconPaths																																		UATCommandLine										AutoSDKPath			SDKStatus						SDKTutorial																								bEnabledForUse													BinaryFolderName	IniPlatformName		FbUsesHostCompiler		bUATClosesAfterLaunch
+	BuildPlatformInfo(TEXT("WindowsNoEditor"),			TEXT("WindowsNoEditor"),	LOCTEXT("WindowsNoEditor", "Windows"),							EPlatformType::Game,	EPlatformFlags::None,			FPlatformIconPaths(TEXT("Launcher/WindowsTarget/Platform_WindowsNoEditor_24x"), TEXT("Launcher/WindowsTarget/Platform_WindowsNoEditor_128x")),	TEXT("-targetplatform=Win64"),						TEXT(""),			EPlatformSDKStatus::Unknown,	TEXT("/Engine/Tutorial/Installation/InstallingVisualStudioTutorial.InstallingVisualStudioTutorial"),	IsAvailableOnWindows,											TEXT("Win64"),		TEXT("Windows"),	IsAvailableOnWindows,	false),
+	BuildPlatformInfo(TEXT("WindowsNoEditor_Win32"),	TEXT("WindowsNoEditor"),	LOCTEXT("WindowsNoEditor_Win32", "Windows (32-bit)"),			EPlatformType::Game,	EPlatformFlags::BuildFlavor,	FPlatformIconPaths(TEXT("Launcher/WindowsTarget/Platform_WindowsNoEditor_24x"), TEXT("Launcher/WindowsTarget/Platform_WindowsNoEditor_128x")),	TEXT("-targetplatform=Win32"),						TEXT(""),			EPlatformSDKStatus::Unknown,	TEXT("/Engine/Tutorial/Installation/InstallingVisualStudioTutorial.InstallingVisualStudioTutorial"),	IsAvailableOnWindows,											TEXT("Win32"),		TEXT("Windows"),	IsAvailableOnWindows,	false),
+	BuildPlatformInfo(TEXT("WindowsNoEditor_Win64"),	TEXT("WindowsNoEditor"),	LOCTEXT("WindowsNoEditor_Win64", "Windows (64-bit)"),			EPlatformType::Game,	EPlatformFlags::BuildFlavor,	FPlatformIconPaths(TEXT("Launcher/WindowsTarget/Platform_WindowsNoEditor_24x"), TEXT("Launcher/WindowsTarget/Platform_WindowsNoEditor_128x")),	TEXT("-targetplatform=Win64"),						TEXT(""),			EPlatformSDKStatus::Unknown,	TEXT("/Engine/Tutorial/Installation/InstallingVisualStudioTutorial.InstallingVisualStudioTutorial"),	IsAvailableOnWindows,											TEXT("Win64"),		TEXT("Windows"),	IsAvailableOnWindows,	false),
+	BuildPlatformInfo(TEXT("Windows"),					TEXT("Windows"),			LOCTEXT("WindowsEditor", "Windows (Editor)"),					EPlatformType::Editor,	EPlatformFlags::None,			FPlatformIconPaths(TEXT("Launcher/WindowsTarget/Platform_Windows_24x"), TEXT("Launcher/WindowsTarget/Platform_Windows_128x")),					TEXT(""),											TEXT(""),			EPlatformSDKStatus::Unknown,	TEXT("/Engine/Tutorial/Installation/InstallingVisualStudioTutorial.InstallingVisualStudioTutorial"),	IsAvailableOnWindows,											TEXT("Win64"),		TEXT("Windows"),	IsAvailableOnWindows,	false),
+	BuildPlatformInfo(TEXT("WindowsClient"),			TEXT("WindowsClient"),		LOCTEXT("WindowsClient", "Windows (Client-only)"),				EPlatformType::Client,	EPlatformFlags::None,			FPlatformIconPaths(TEXT("Launcher/WindowsTarget/Platform_Windows_24x"), TEXT("Launcher/WindowsTarget/Platform_Windows_128x")),					TEXT("-client"),									TEXT(""),			EPlatformSDKStatus::Unknown,	TEXT("/Engine/Tutorial/Installation/InstallingVisualStudioTutorial.InstallingVisualStudioTutorial"),	IsAvailableOnWindows,											TEXT("Win64"),		TEXT("Windows"),	IsAvailableOnWindows,	false),
+	BuildPlatformInfo(TEXT("WindowsServer"),			TEXT("WindowsServer"),		LOCTEXT("WindowsServer", "Windows (Dedicated Server)"),			EPlatformType::Server,	EPlatformFlags::None,			FPlatformIconPaths(TEXT("Launcher/WindowsTarget/Platform_WindowsServer_24x"), TEXT("Launcher/WindowsTarget/Platform_WindowsServer_128x")),		TEXT(""),											TEXT(""),			EPlatformSDKStatus::Unknown,	TEXT("/Engine/Tutorial/Installation/InstallingVisualStudioTutorial.InstallingVisualStudioTutorial"),	IsAvailableOnWindows,											TEXT("Win64"),		TEXT("Windows"),	IsAvailableOnWindows,	false),
+
+	BuildPlatformInfo(TEXT("MacNoEditor"),				TEXT("MacNoEditor"),		LOCTEXT("MacNoEditor", "Mac"),									EPlatformType::Game,	EPlatformFlags::None,			FPlatformIconPaths(TEXT("Launcher/MacTarget/Platform_Mac_24x"), TEXT("Launcher/MacTarget/Platform_Mac_128x")),									TEXT("-targetplatform=Mac"),						TEXT(""),			EPlatformSDKStatus::Unknown,	TEXT("/Engine/Tutorial/Installation/InstallingXCodeTutorial.InstallingXCodeTutorial"),					IsAvailableOnMac,												TEXT("Mac"),		TEXT("Mac"),		IsAvailableOnMac,		true),
+	BuildPlatformInfo(TEXT("Mac"),						TEXT("Mac"),				LOCTEXT("MacEditor", "Mac (Editor)"),							EPlatformType::Editor,	EPlatformFlags::None,			FPlatformIconPaths(TEXT("Launcher/MacTarget/Platform_Mac_24x"), TEXT("Launcher/MacTarget/Platform_Mac_128x")),									TEXT(""),											TEXT(""),			EPlatformSDKStatus::Unknown,	TEXT("/Engine/Tutorial/Installation/InstallingXCodeTutorial.InstallingXCodeTutorial"),					IsAvailableOnMac,												TEXT("Mac"),		TEXT("Mac"),		IsAvailableOnMac,		true),
+	BuildPlatformInfo(TEXT("MacClient"),				TEXT("MacClient"),			LOCTEXT("MacClient", "Mac (Client-only)"),						EPlatformType::Client,	EPlatformFlags::None,			FPlatformIconPaths(TEXT("Launcher/MacTarget/Platform_Mac_24x"), TEXT("Launcher/MacTarget/Platform_Mac_128x")),									TEXT("-client"),									TEXT(""),			EPlatformSDKStatus::Unknown,	TEXT("/Engine/Tutorial/Installation/InstallingXCodeTutorial.InstallingXCodeTutorial"),					IsAvailableOnMac,												TEXT("Mac"),		TEXT("Mac"),		IsAvailableOnMac,		true),
+	BuildPlatformInfo(TEXT("MacServer"),				TEXT("MacServer"),			LOCTEXT("MacServer", "Mac (Dedicated Server)"),					EPlatformType::Server,	EPlatformFlags::None,			FPlatformIconPaths(TEXT("Launcher/MacTarget/Platform_Mac_24x"), TEXT("Launcher/MacTarget/Platform_Mac_128x")),									TEXT(""),											TEXT(""),			EPlatformSDKStatus::Unknown,	TEXT("/Engine/Tutorial/Installation/InstallingXCodeTutorial.InstallingXCodeTutorial"),					IsAvailableOnMac,												TEXT("Mac"),		TEXT("Mac"),		IsAvailableOnMac,		true),
+
+	BuildPlatformInfo(TEXT("LinuxNoEditor"),			TEXT("LinuxNoEditor"),		LOCTEXT("LinuxNoEditor", "Linux"),								EPlatformType::Game,	EPlatformFlags::None,			FPlatformIconPaths(TEXT("Launcher/LinuxTarget/Platform_Linux_24x"), TEXT("Launcher/LinuxTarget/Platform_Linux_128x")),							TEXT("-targetplatform=Linux"),						TEXT("Linux_x64"),	EPlatformSDKStatus::Unknown,	TEXT("/Platforms/Linux/GettingStarted"),																IsAvailableOnLinux || IsAvailableOnWindows,						TEXT("Linux"),		TEXT("Linux"),		IsAvailableOnLinux,		true),	
+	BuildPlatformInfo(TEXT("Linux"),					TEXT("Linux"),				LOCTEXT("LinuxEditor", "Linux (Editor)"),						EPlatformType::Editor,	EPlatformFlags::None,			FPlatformIconPaths(TEXT("Launcher/LinuxTarget/Platform_Linux_24x"), TEXT("Launcher/LinuxTarget/Platform_Linux_128x")),							TEXT(""),											TEXT("Linux_x64"),	EPlatformSDKStatus::Unknown,	TEXT("/Platforms/Linux/GettingStarted"),																IsAvailableOnLinux,												TEXT("Linux"),		TEXT("Linux"),		IsAvailableOnLinux,		true),
+	BuildPlatformInfo(TEXT("LinuxClient"),				TEXT("LinuxClient"),		LOCTEXT("LinuxClient", "Linux (Client-only)"),					EPlatformType::Client,	EPlatformFlags::None,			FPlatformIconPaths(TEXT("Launcher/LinuxTarget/Platform_Linux_24x"), TEXT("Launcher/LinuxTarget/Platform_Linux_128x")),							TEXT("-client"),									TEXT("Linux_x64"),	EPlatformSDKStatus::Unknown,	TEXT("/Platforms/Linux/GettingStarted"),																IsAvailableOnLinux || IsAvailableOnWindows,						TEXT("Linux"),		TEXT("Linux"),		IsAvailableOnLinux,		true),
+	BuildPlatformInfo(TEXT("LinuxServer"),				TEXT("LinuxServer"),		LOCTEXT("LinuxServer", "Linux (Dedicated Server)"),				EPlatformType::Server,	EPlatformFlags::None,			FPlatformIconPaths(TEXT("Launcher/LinuxTarget/Platform_Linux_24x"), TEXT("Launcher/LinuxTarget/Platform_Linux_128x")),							TEXT(""),											TEXT("Linux_x64"),	EPlatformSDKStatus::Unknown,	TEXT("/Platforms/Linux/GettingStarted"),																IsAvailableOnLinux || IsAvailableOnWindows,						TEXT("Linux"),		TEXT("Linux"),		IsAvailableOnLinux,		true),
+
+	BuildPlatformInfo(TEXT("IOS"),						TEXT("IOS"),				LOCTEXT("IOS", "iOS"),											EPlatformType::Game,	EPlatformFlags::None,			FPlatformIconPaths(TEXT("Launcher/iOSTarget/Platform_iOS_24x"), TEXT("Launcher/iOSTarget/Platform_iOS_128x")),									TEXT(""),											TEXT(""),			EPlatformSDKStatus::Unknown,	IOSTutorial,																							IsAvailableOnWindows || IsAvailableOnMac,						TEXT("IOS"),		TEXT("IOS"),		false,					true),
+
+	BuildPlatformInfo(TEXT("Android"),					TEXT("Android"),			LOCTEXT("Android", "Android"),									EPlatformType::Game,	EPlatformFlags::None,			FPlatformIconPaths(TEXT("Launcher/AndroidTarget/Platform_Android_24x"), TEXT("Launcher/AndroidTarget/Platform_Android_128x")),					TEXT(""),											TEXT("Android"),	EPlatformSDKStatus::Unknown,	TEXT("/Engine/Tutorial/Mobile/SettingUpAndroidTutorial.SettingUpAndroidTutorial"),						IsAvailableOnLinux || IsAvailableOnWindows || IsAvailableOnMac,	TEXT("Android"),	TEXT("Android"),	false,					true),
+	BuildPlatformInfo(TEXT("Android_All"),				TEXT("Android"),			LOCTEXT("Android_All", "Android (All)"),						EPlatformType::Game,	EPlatformFlags::CookFlavor,		FPlatformIconPaths(TEXT("Launcher/AndroidTarget/Platform_Android_24x"), TEXT("Launcher/AndroidTarget/Platform_Android_128x")),					TEXT(""),											TEXT("Android"),	EPlatformSDKStatus::Unknown,	TEXT("/Engine/Tutorial/Mobile/SettingUpAndroidTutorial.SettingUpAndroidTutorial"),						IsAvailableOnLinux || IsAvailableOnWindows || IsAvailableOnMac,	TEXT("Android"),	TEXT("Android"),	false,					true),
+	BuildPlatformInfo(TEXT("Android_Multi"),			TEXT("Android_Multi"),		LOCTEXT("Android_Multi", "Android (Multi)"),					EPlatformType::Game,	EPlatformFlags::CookFlavor,		FPlatformIconPaths(TEXT("Launcher/AndroidTarget/Platform_Android_24x"), TEXT("Launcher/AndroidTarget/Platform_Android_128x")),					TEXT("-targetplatform=Android -cookflavor=Multi"),	TEXT("Android"),	EPlatformSDKStatus::Unknown,	TEXT("/Engine/Tutorial/Mobile/SettingUpAndroidTutorial.SettingUpAndroidTutorial"),						IsAvailableOnLinux || IsAvailableOnWindows || IsAvailableOnMac,	TEXT("Android"),	TEXT("Android"),	false,					true),
+	BuildPlatformInfo(TEXT("Android_ATC"),				TEXT("Android_ATC"),		LOCTEXT("Android_ATC", "Android (ATC)"),						EPlatformType::Game,	EPlatformFlags::CookFlavor,		FPlatformIconPaths(TEXT("Launcher/AndroidTarget/Platform_Android_ATC_24x"), TEXT("Launcher/AndroidTarget/Platform_Android_128x")),				TEXT("-targetplatform=Android -cookflavor=ATC"),	TEXT("Android"),	EPlatformSDKStatus::Unknown,	TEXT("/Engine/Tutorial/Mobile/SettingUpAndroidTutorial.SettingUpAndroidTutorial"),						IsAvailableOnLinux || IsAvailableOnWindows || IsAvailableOnMac,	TEXT("Android"),	TEXT("Android"),	false,					true),
+	BuildPlatformInfo(TEXT("Android_DXT"),				TEXT("Android_DXT"),		LOCTEXT("Android_DXT", "Android (DXT)"),						EPlatformType::Game,	EPlatformFlags::CookFlavor,		FPlatformIconPaths(TEXT("Launcher/AndroidTarget/Platform_Android_DXT_24x"), TEXT("Launcher/AndroidTarget/Platform_Android_128x")),				TEXT("-targetplatform=Android -cookflavor=DXT"),	TEXT("Android"),	EPlatformSDKStatus::Unknown,	TEXT("/Engine/Tutorial/Mobile/SettingUpAndroidTutorial.SettingUpAndroidTutorial"),						IsAvailableOnLinux || IsAvailableOnWindows || IsAvailableOnMac,	TEXT("Android"),	TEXT("Android"),	false,					true),
+	BuildPlatformInfo(TEXT("Android_ETC1"),				TEXT("Android_ETC1"),		LOCTEXT("Android_ETC1", "Android (ETC1)"),						EPlatformType::Game,	EPlatformFlags::CookFlavor,		FPlatformIconPaths(TEXT("Launcher/AndroidTarget/Platform_Android_ETC1_24x"), TEXT("Launcher/AndroidTarget/Platform_Android_128x")),				TEXT("-targetplatform=Android -cookflavor=ETC1"),	TEXT("Android"),	EPlatformSDKStatus::Unknown,	TEXT("/Engine/Tutorial/Mobile/SettingUpAndroidTutorial.SettingUpAndroidTutorial"),						IsAvailableOnLinux || IsAvailableOnWindows || IsAvailableOnMac,	TEXT("Android"),	TEXT("Android"),	false,					true),
+	BuildPlatformInfo(TEXT("Android_ETC2"),				TEXT("Android_ETC2"),		LOCTEXT("Android_ETC2", "Android (ETC2)"),						EPlatformType::Game,	EPlatformFlags::CookFlavor,		FPlatformIconPaths(TEXT("Launcher/AndroidTarget/Platform_Android_ETC2_24x"), TEXT("Launcher/AndroidTarget/Platform_Android_128x")),				TEXT("-targetplatform=Android -cookflavor=ETC2"),	TEXT("Android"),	EPlatformSDKStatus::Unknown,	TEXT("/Engine/Tutorial/Mobile/SettingUpAndroidTutorial.SettingUpAndroidTutorial"),						IsAvailableOnLinux || IsAvailableOnWindows || IsAvailableOnMac,	TEXT("Android"),	TEXT("Android"),	false,					true),
+	BuildPlatformInfo(TEXT("Android_PVRTC"),			TEXT("Android_PVRTC"),		LOCTEXT("Android_PVRTC", "Android (PVRTC)"),					EPlatformType::Game,	EPlatformFlags::CookFlavor,		FPlatformIconPaths(TEXT("Launcher/AndroidTarget/Platform_Android_PVRTC_24x"), TEXT("Launcher/AndroidTarget/Platform_Android_128x")),			TEXT("-targetplatform=Android -cookflavor=PVRTC"),	TEXT("Android"),	EPlatformSDKStatus::Unknown,	TEXT("/Engine/Tutorial/Mobile/SettingUpAndroidTutorial.SettingUpAndroidTutorial"),						IsAvailableOnLinux || IsAvailableOnWindows || IsAvailableOnMac,	TEXT("Android"),	TEXT("Android"),	false,					true),
+	BuildPlatformInfo(TEXT("Android_ASTC"),				TEXT("Android_ASTC"),		LOCTEXT("Android_ASTC", "Android (ASTC)"),						EPlatformType::Game,	EPlatformFlags::CookFlavor,		FPlatformIconPaths(TEXT("Launcher/AndroidTarget/Platform_Android_ASTC_24x"), TEXT("Launcher/AndroidTarget/Platform_Android_128x")),				TEXT("-targetplatform=Android -cookflavor=ASTC"),	TEXT("Android"),	EPlatformSDKStatus::Unknown,	TEXT("/Engine/Tutorial/Mobile/SettingUpAndroidTutorial.SettingUpAndroidTutorial"),						IsAvailableOnLinux || IsAvailableOnWindows || IsAvailableOnMac,	TEXT("Android"),	TEXT("Android"),	false,					true),
+
+	BuildPlatformInfo(TEXT("HTML5"),					TEXT("HTML5"),				LOCTEXT("HTML5", "HTML5"),										EPlatformType::Game,	EPlatformFlags::None,			FPlatformIconPaths(TEXT("Launcher/HTML5Target/Platform_HTML5_24x"), TEXT("Launcher/HTML5Target/Platform_HTML5_128x")),							TEXT(""),											TEXT("HTML5"),		EPlatformSDKStatus::Unknown,	TEXT("/Platforms/HTML5/GettingStarted"),																IsAvailableOnWindows || IsAvailableOnMac,						TEXT("HTML5"),		TEXT("HTML5"),		false,					true),
+
+	BuildPlatformInfo(TEXT("PS4"),						TEXT("PS4"),				LOCTEXT("PS4", "PlayStation 4"),								EPlatformType::Game,	EPlatformFlags::None,			FPlatformIconPaths(TEXT("Launcher/PS4Target/Platform_PS4_24x"), TEXT("Launcher/PS4Target/Platform_PS4_128x")),									TEXT(""),											TEXT("PS4"),		EPlatformSDKStatus::Unknown,	TEXT("/Platforms/PS4/GettingStarted"),																	IsAvailableOnWindows,											TEXT("PS4"),		TEXT("PS4"),		false,					false),
+
+	BuildPlatformInfo(TEXT("XboxOne"),					TEXT("XboxOne"),			LOCTEXT("XboxOne", "Xbox One"),									EPlatformType::Game,	EPlatformFlags::None,			FPlatformIconPaths(TEXT("Launcher/XboxOneTarget/Platform_XboxOne_24x"), TEXT("Launcher/XboxOneTarget/Platform_XboxOne_128x")),					TEXT(""),											TEXT("XboxOne"),	EPlatformSDKStatus::Unknown,	TEXT("/Platforms/XboxOne/GettingStarted"),																IsAvailableOnWindows,											TEXT("XboxOne"),	TEXT("XboxOne"),	false,					true),
+
+	BuildPlatformInfo(TEXT("AllDesktop"),				TEXT("AllDesktop"),			LOCTEXT("DesktopTargetPlatDisplay", "Desktop (Win+Mac+Linux)"),	EPlatformType::Game,	EPlatformFlags::None,			FPlatformIconPaths(TEXT("Launcher/DesktopTarget/Platform_Desktop_24x"), TEXT("Launcher/DesktopTarget/Platform_Desktop_128x")),					TEXT(""),											TEXT(""),			EPlatformSDKStatus::Unknown,	TEXT(""),																								IsAvailableOnWindows /* see note below */,						TEXT(""),			TEXT(""),			false,					true),
+
+	BuildPlatformInfo(TEXT("TVOS"),						TEXT("TVOS"),				LOCTEXT("TVOSTargetPlatDisplay", "tvOS"),						EPlatformType::Game,	EPlatformFlags::None,			FPlatformIconPaths(TEXT("Launcher/TVOSTarget/Platform_TVOS_24x"), TEXT("Launcher/TVOSTarget/Platform_TVOS_128x")),								TEXT(""),											TEXT(""),			EPlatformSDKStatus::Unknown,	TEXT(""),																								IsAvailableOnWindows || IsAvailableOnMac,						TEXT("TVOS"),		TEXT("TVOS"),		false,					true),
+	BuildPlatformInfo(TEXT("Switch"),					TEXT("Switch"),				LOCTEXT("Switch", "Switch"),									EPlatformType::Game,	EPlatformFlags::None,			FPlatformIconPaths(TEXT("Launcher/SwitchTarget/Platform_Switch_24x"), TEXT("Launcher/SwitchTarget/Platform_Switch_128x")),						TEXT(""),											TEXT(""),			EPlatformSDKStatus::Unknown,	TEXT(""),																								IsAvailableOnWindows,											TEXT("Switch"),		TEXT("Switch"),		false,					true),
+
+	// Note: For "AllDesktop" bEnabledForUse value, see SProjectTargetPlatformSettings::Construct !!!! IsAvailableOnWindows || IsAvailableOnMac || IsAvailableOnLinux
 };
 
 } // anonymous namespace
diff --git a/Engine/Source/Developer/DesktopPlatform/Public/PlatformInfo.h b/Engine/Source/Developer/DesktopPlatform/Public/PlatformInfo.h
--- a/Engine/Source/Developer/DesktopPlatform/Public/PlatformInfo.h
+++ b/Engine/Source/Developer/DesktopPlatform/Public/PlatformInfo.h
@@ -157,6 +157,9 @@ namespace PlatformInfo
 		/** Name of sub-folder where binaries will be placed */
 		FString BinaryFolderName;
 
+		/** Name of this platform when loading INI files */
+		FString IniPlatformName;
+
 		/** Enabled for use */
 		bool bEnabledForUse;
 
diff --git a/Engine/Source/Editor/AnimGraph/Classes/AnimGraphNode_TwoWayBlend.h b/Engine/Source/Editor/AnimGraph/Classes/AnimGraphNode_TwoWayBlend.h
--- a/Engine/Source/Editor/AnimGraph/Classes/AnimGraphNode_TwoWayBlend.h
+++ b/Engine/Source/Editor/AnimGraph/Classes/AnimGraphNode_TwoWayBlend.h
@@ -14,7 +14,7 @@ class UAnimGraphNode_TwoWayBlend : public UAnimGraphNode_Base
 	GENERATED_UCLASS_BODY()
 
 	UPROPERTY(EditAnywhere, Category=Settings)
-	FAnimationNode_TwoWayBlend BlendNode;
+	FAnimNode_TwoWayBlend BlendNode;
 
 	//~ Begin UEdGraphNode Interface.
 	virtual FLinearColor GetNodeTitleColor() const override;
diff --git a/Engine/Source/Editor/AnimationBlueprintEditor/Private/AnimationBlueprintEditor.cpp b/Engine/Source/Editor/AnimationBlueprintEditor/Private/AnimationBlueprintEditor.cpp
--- a/Engine/Source/Editor/AnimationBlueprintEditor/Private/AnimationBlueprintEditor.cpp
+++ b/Engine/Source/Editor/AnimationBlueprintEditor/Private/AnimationBlueprintEditor.cpp
@@ -35,6 +35,7 @@
 #include "AnimGraphNode_SequenceEvaluator.h"
 #include "AnimGraphNode_PoseByName.h"
 #include "AnimGraphNode_PoseBlendNode.h"
+#include "AnimGraphNode_MultiWayBlend.h"
 
 #include "Animation/AnimNotifies/AnimNotifyState.h"
 
@@ -94,7 +95,7 @@ public:
 	{
 		AnimationBlueprintEditorPtr = InAnimationBlueprintEditor;
 
-		SSingleObjectDetailsPanel::Construct(SSingleObjectDetailsPanel::FArguments().HostCommandList(InAnimationBlueprintEditor->GetToolkitCommands()), /*bAutomaticallyObserveViaGetObjectToObserve*/ true, /*bAllowSearch*/ true);
+		SSingleObjectDetailsPanel::Construct(SSingleObjectDetailsPanel::FArguments().HostCommandList(InAnimationBlueprintEditor->GetToolkitCommands()).HostTabManager(InAnimationBlueprintEditor->GetTabManager()), /*bAutomaticallyObserveViaGetObjectToObserve*/ true, /*bAllowSearch*/ true);
 
 		PropertyView->SetIsPropertyEditingEnabledDelegate(FIsPropertyEditingEnabled::CreateStatic([] { return !GIntraFrameDebuggingGameThread; }));
 	}
@@ -385,6 +386,11 @@ void FAnimationBlueprintEditor::OnAddPosePin()
 				FilterNode->AddPinToBlendByFilter();
 				break;
 			}
+			else if (UAnimGraphNode_MultiWayBlend* MultiBlendNode = Cast<UAnimGraphNode_MultiWayBlend>(Node))
+			{
+				MultiBlendNode->AddPinToBlendNode();
+				break;
+			}
 		}
 	}
 }
@@ -397,8 +403,9 @@ bool FAnimationBlueprintEditor::CanAddPosePin() const
 void FAnimationBlueprintEditor::OnRemovePosePin()
 {
 	const FGraphPanelSelectionSet SelectedNodes = GetSelectedNodes();
-	UAnimGraphNode_BlendListByInt* BlendListIntNode = NULL;
-	UAnimGraphNode_LayeredBoneBlend* BlendByFilterNode = NULL;
+	UAnimGraphNode_BlendListByInt* BlendListIntNode = nullptr;
+	UAnimGraphNode_LayeredBoneBlend* BlendByFilterNode = nullptr;
+	UAnimGraphNode_MultiWayBlend* BlendByMultiway = nullptr;
 
 	if (SelectedNodes.Num() == 1)
 	{
@@ -414,6 +421,11 @@ void FAnimationBlueprintEditor::OnRemovePosePin()
 				BlendByFilterNode = LayeredBlendNode;
 				break;
 			}		
+			else if (UAnimGraphNode_MultiWayBlend* MultiwayBlendNode = Cast<UAnimGraphNode_MultiWayBlend>(*NodeIt))
+			{
+				BlendByMultiway = MultiwayBlendNode;
+				break;
+			}
 		}
 	}
 
@@ -445,6 +457,17 @@ void FAnimationBlueprintEditor::OnRemovePosePin()
 			// Update the graph so that the node will be refreshed
 			FocusedGraphEd->NotifyGraphChanged();
 		}
+
+		if (BlendByMultiway)
+		{
+			// make sure we at least have BlendListNode selected
+			UEdGraphPin* SelectedPin = FocusedGraphEd->GetGraphPinForMenu();
+
+			BlendByMultiway->RemovePinFromBlendNode(SelectedPin);
+
+			// Update the graph so that the node will be refreshed
+			FocusedGraphEd->NotifyGraphChanged();
+		}
 	}
 }
 
diff --git a/Engine/Source/Editor/DetailCustomizations/Private/DetailCustomizations.cpp b/Engine/Source/Editor/DetailCustomizations/Private/DetailCustomizations.cpp
--- a/Engine/Source/Editor/DetailCustomizations/Private/DetailCustomizations.cpp
+++ b/Engine/Source/Editor/DetailCustomizations/Private/DetailCustomizations.cpp
@@ -102,8 +102,11 @@
 #include "MovieSceneCaptureCustomization.h"
 #include "MovieSceneEvalOptionsCustomization.h"
 #include "MovieSceneEventParametersCustomization.h"
+#include "MovieSceneSequencePlaybackSettingsCustomization.h"
 #include "MovieSceneCurveInterfaceKeyEditStructCustomization.h"
 #include "LevelSequenceBurnInOptionsCustomization.h"
+#include "MovieSceneBindingOverrideDataCustomization.h"
+#include "MovieSceneObjectBindingPtrCustomization.h"
 #include "TextCustomization.h"
 #include "AnimTrailNodeDetails.h"
 #include "MaterialProxySettingsCustomizations.h"
@@ -231,6 +234,9 @@ void FDetailCustomizationsModule::RegisterPropertyTypeCustomizations()
 	RegisterCustomPropertyTypeLayout("CameraFilmbackSettings", FOnGetPropertyTypeCustomizationInstance::CreateStatic(&FCameraFilmbackSettingsCustomization::MakeInstance));
 	RegisterCustomPropertyTypeLayout("CameraLensSettings", FOnGetPropertyTypeCustomizationInstance::CreateStatic(&FCameraLensSettingsCustomization::MakeInstance));
 	RegisterCustomPropertyTypeLayout("CameraFocusSettings", FOnGetPropertyTypeCustomizationInstance::CreateStatic(&FCameraFocusSettingsCustomization::MakeInstance));
+	RegisterCustomPropertyTypeLayout("MovieSceneSequencePlaybackSettings", FOnGetPropertyTypeCustomizationInstance::CreateStatic(&FMovieSceneSequencePlaybackSettingsCustomization::MakeInstance));
+	RegisterCustomPropertyTypeLayout("MovieSceneBindingOverrideData", FOnGetPropertyTypeCustomizationInstance::CreateStatic(&FMovieSceneBindingOverrideDataCustomization::MakeInstance));
+	RegisterCustomPropertyTypeLayout("MovieSceneObjectBindingPtr", FOnGetPropertyTypeCustomizationInstance::CreateStatic(&FMovieSceneObjectBindingPtrCustomization::MakeInstance));
 	RegisterCustomPropertyTypeLayout("MovieSceneTrackEvalOptions", FOnGetPropertyTypeCustomizationInstance::CreateStatic(&FMovieSceneTrackEvalOptionsCustomization::MakeInstance));
 	RegisterCustomPropertyTypeLayout("MovieSceneSectionEvalOptions", FOnGetPropertyTypeCustomizationInstance::CreateStatic(&FMovieSceneSectionEvalOptionsCustomization::MakeInstance));
 	RegisterCustomPropertyTypeLayout("MovieSceneEventParameters", FOnGetPropertyTypeCustomizationInstance::CreateStatic(&FMovieSceneEventParametersCustomization::MakeInstance));
diff --git a/Engine/Source/Editor/DetailCustomizations/Private/MovieSceneEventParametersCustomization.cpp b/Engine/Source/Editor/DetailCustomizations/Private/MovieSceneEventParametersCustomization.cpp
--- a/Engine/Source/Editor/DetailCustomizations/Private/MovieSceneEventParametersCustomization.cpp
+++ b/Engine/Source/Editor/DetailCustomizations/Private/MovieSceneEventParametersCustomization.cpp
@@ -59,13 +59,13 @@ void FMovieSceneEventParametersCustomization::CustomizeChildren(TSharedRef<IProp
 
 	if (EditStructData->GetStructMemory())
 	{
-		FSimpleDelegate OnEditStructContentsChangedDelegate = FSimpleDelegate::CreateSP(this, &FMovieSceneEventParametersCustomization::OnEditStructContentsChanged);
+		FSimpleDelegate OnEditStructChildContentsChangedDelegate = FSimpleDelegate::CreateSP(this, &FMovieSceneEventParametersCustomization::OnEditStructChildContentsChanged);
 
 		TArray<TSharedPtr<IPropertyHandle>> ExternalHandles = ChildBuilder.AddChildStructure(EditStructData.ToSharedRef());
 		for (TSharedPtr<IPropertyHandle> Handle : ExternalHandles)
 		{
-			Handle->SetOnPropertyValueChanged(OnEditStructContentsChangedDelegate);
-			Handle->SetOnChildPropertyValueChanged(OnEditStructContentsChangedDelegate);
+			Handle->SetOnPropertyValueChanged(OnEditStructChildContentsChangedDelegate);
+			Handle->SetOnChildPropertyValueChanged(OnEditStructChildContentsChangedDelegate);
 		}
 	}
 }
@@ -97,7 +97,7 @@ void FMovieSceneEventParametersCustomization::OnStructChanged(const FAssetData&
 	PropertyUtilities->ForceRefresh();
 }
 
-void FMovieSceneEventParametersCustomization::OnEditStructContentsChanged()
+void FMovieSceneEventParametersCustomization::OnEditStructChildContentsChanged()
 {
 	// @todo: call modify on the outer object if possible
 	UScriptStruct* Struct = Cast<UScriptStruct>((UStruct*)EditStructData->GetStruct());
diff --git a/Engine/Source/Editor/DetailCustomizations/Private/MovieSceneEventParametersCustomization.h b/Engine/Source/Editor/DetailCustomizations/Private/MovieSceneEventParametersCustomization.h
--- a/Engine/Source/Editor/DetailCustomizations/Private/MovieSceneEventParametersCustomization.h
+++ b/Engine/Source/Editor/DetailCustomizations/Private/MovieSceneEventParametersCustomization.h
@@ -22,7 +22,7 @@ public:
 
 private:
 	void OnStructChanged(const FAssetData& AssetData);
-	void OnEditStructContentsChanged();
+	void OnEditStructChildContentsChanged();
 
 	TSharedPtr<IPropertyUtilities> PropertyUtilities;
 	TSharedPtr<IPropertyHandle> PropertyHandle;
diff --git a/Engine/Source/Editor/DetailCustomizations/Private/MovieSceneObjectBindingPtrCustomization.cpp b/Engine/Source/Editor/DetailCustomizations/Private/MovieSceneObjectBindingPtrCustomization.cpp
new file mode 100644
--- /dev/null
+++ b/Engine/Source/Editor/DetailCustomizations/Private/MovieSceneObjectBindingPtrCustomization.cpp
@@ -0,0 +1,191 @@
+// Copyright 1998-2016 Epic Games, Inc. All Rights Reserved.
+
+#include "MovieSceneObjectBindingPtrCustomization.h"
+#include "IPropertyUtilities.h"
+#include "MovieSceneBindingOwnerInterface.h"
+#include "MovieSceneSequence.h"
+#include "MovieScene.h"
+#include "PropertyHandle.h"
+#include "DetailWidgetRow.h"
+#include "IDetailChildrenBuilder.h"
+#include "Framework/Application/SlateApplication.h"
+#include "Widgets/Input/SComboButton.h"
+#include "Widgets/Text/STextBlock.h"
+#include "Widgets/Views/STableRow.h"
+#include "Widgets/Views/STreeView.h"
+
+
+#define LOCTEXT_NAMESPACE "MovieSceneObjectBindingPtrCustomization"
+
+
+namespace
+{
+	DECLARE_DELEGATE_OneParam(FOnSelectionChanged, FGuid);
+
+	class SObjectBindingTree : public SCompoundWidget
+	{
+	public:
+		SLATE_BEGIN_ARGS(SObjectBindingTree){}
+			SLATE_EVENT(FOnSelectionChanged, OnSelectionChanged)
+		SLATE_END_ARGS()
+
+		void Construct(const FArguments& InArgs, TWeakObjectPtr<UMovieSceneSequence> InSequence)
+		{
+			OnSelectionChangedEvent = InArgs._OnSelectionChanged;
+			WeakSequence = InSequence;
+
+			UMovieScene* MovieScene = WeakSequence->GetMovieScene();
+			if (!MovieScene)
+			{
+				ChildSlot
+				[
+					SNew(STextBlock).Text(LOCTEXT("InvalidSequence", "Invalid sequence"))
+				];
+				return;
+			}
+
+			const int32 Count = MovieScene->GetPossessableCount();
+			for (int32 Index = 0; Index < Count; ++Index)
+			{
+				const FMovieScenePossessable& Possessable = MovieScene->GetPossessable(Index);
+				if (!WeakSequence.Get()->CanRebindPossessable(Possessable))
+				{
+					continue;
+				}
+
+				FGuid ParentGuid = Possessable.GetParent();
+				Hierarchy.FindOrAdd(ParentGuid).Add(MakeShared<FGuid>(Possessable.GetGuid()));
+			}
+
+			RootGuids = Hierarchy.FindRef(FGuid());
+
+			ChildSlot
+			[
+				SNew(STreeView<TSharedPtr<FGuid>>)
+				.TreeItemsSource(&RootGuids)
+				.OnGenerateRow(this, &SObjectBindingTree::OnGenerateRowWidget)
+				.OnGetChildren(this, &SObjectBindingTree::OnGetChildrenForTree)
+				.OnSelectionChanged(this, &SObjectBindingTree::OnTreeSelectionChanged)
+			];
+		}
+
+		TSharedRef<ITableRow> OnGenerateRowWidget(TSharedPtr<FGuid> InGuid, const TSharedRef<STableViewBase>& OwnerTable)
+		{
+			UMovieSceneSequence* Sequence = WeakSequence.Get();
+			UMovieScene* MovieScene = Sequence ? Sequence->GetMovieScene() : nullptr;
+			FText DisplayString = MovieScene ? MovieScene->GetObjectDisplayName(*InGuid) : FText();
+
+			return SNew(STableRow<TSharedPtr<FGuid>>, OwnerTable)
+			[
+				SNew(STextBlock).Text(DisplayString)
+			];
+		}
+
+		void OnGetChildrenForTree(TSharedPtr<FGuid> ParentId, TArray<TSharedPtr<FGuid>>& OutChildren)
+		{
+			OutChildren = Hierarchy.FindRef(*ParentId);
+		}
+
+		void OnTreeSelectionChanged(TSharedPtr<FGuid> SelectedGuid, ESelectInfo::Type)
+		{
+			OnSelectionChangedEvent.ExecuteIfBound(*SelectedGuid);
+			FSlateApplication::Get().DismissAllMenus();
+		}
+
+	private:
+		
+		FOnSelectionChanged OnSelectionChangedEvent;
+
+		TArray<TSharedPtr<FGuid>> RootGuids;
+		TMap<FGuid, TArray<TSharedPtr<FGuid>>> Hierarchy;
+		TWeakObjectPtr<UMovieSceneSequence> WeakSequence;
+	};
+}
+
+void FMovieSceneObjectBindingPtrCustomization::CustomizeHeader(TSharedRef<IPropertyHandle> PropertyHandle, FDetailWidgetRow& HeaderRow, IPropertyTypeCustomizationUtils& CustomizationUtils)
+{
+	StructProperty = PropertyHandle;
+	GuidProperty = StructProperty->GetChildHandle(GET_MEMBER_NAME_CHECKED(FMovieSceneObjectBindingPtr, Guid));
+
+	check(GuidProperty.IsValid());
+
+	HeaderRow
+	.NameContent()
+	[
+		StructProperty->CreatePropertyNameWidget()
+	]
+	.ValueContent()
+	[
+		SNew(SComboButton)
+		.OnGetMenuContent(this, &FMovieSceneObjectBindingPtrCustomization::GetBindingMenu)
+		.ContentPadding(FMargin(4.0, 2.0))
+		.ButtonContent()
+		[
+			SNew(STextBlock)
+			.Text(this, &FMovieSceneObjectBindingPtrCustomization::ResolveBindingNameInSequence)
+			.Font(CustomizationUtils.GetRegularFont())
+		]
+	];
+}
+IMovieSceneBindingOwnerInterface* FMovieSceneObjectBindingPtrCustomization::GetInterface() const
+{
+	TArray<UObject*> OuterObjects;
+	StructProperty->GetOuterObjects(OuterObjects);
+
+	if (OuterObjects.Num() != 1)
+	{
+		return nullptr;
+	}
+
+	for ( UObject* NextOuter = OuterObjects[0]; NextOuter; NextOuter = NextOuter->GetOuter() )
+	{
+		if (IMovieSceneBindingOwnerInterface* Result = Cast<IMovieSceneBindingOwnerInterface>(NextOuter))
+		{
+			return Result;
+		}
+	}
+	return nullptr;
+}
+
+UMovieSceneSequence* FMovieSceneObjectBindingPtrCustomization::GetSequence() const
+{
+	IMovieSceneBindingOwnerInterface* OwnerInterface = GetInterface();
+	return OwnerInterface ? OwnerInterface->RetrieveOwnedSequence() : nullptr;
+}
+
+FGuid FMovieSceneObjectBindingPtrCustomization::GetGuid() const
+{
+	TArray<void*> Guids;
+	GuidProperty->AccessRawData(Guids);
+
+	return ensure(Guids.Num() == 1) ? *static_cast<FGuid*>(Guids[0]) : FGuid();
+}
+
+void FMovieSceneObjectBindingPtrCustomization::SetBindingId(FGuid InGuid)
+{
+	TArray<void*> Guids;
+	GuidProperty->AccessRawData(Guids);
+
+	if (ensure(Guids.Num() == 1))
+	{
+		*static_cast<FGuid*>(Guids[0]) = InGuid;
+	}
+}
+
+FText FMovieSceneObjectBindingPtrCustomization::ResolveBindingNameInSequence() const
+{
+	UMovieSceneSequence* Sequence = GetSequence();
+	UMovieScene* MovieScene = Sequence ? Sequence->GetMovieScene() : nullptr;
+
+	FText DisplayString = MovieScene ? MovieScene->GetObjectDisplayName(GetGuid()) : FText();
+	return DisplayString.IsEmpty() ? LOCTEXT("UnresolvedBinding", "Unresolved Binding") : DisplayString;
+}
+
+TSharedRef<SWidget> FMovieSceneObjectBindingPtrCustomization::GetBindingMenu()
+{
+	return
+		SNew(SObjectBindingTree, GetSequence())
+		.OnSelectionChanged(this, &FMovieSceneObjectBindingPtrCustomization::SetBindingId);
+}
+
+#undef LOCTEXT_NAMESPACE
diff --git a/Engine/Source/Editor/DetailCustomizations/Private/MovieSceneObjectBindingPtrCustomization.h b/Engine/Source/Editor/DetailCustomizations/Private/MovieSceneObjectBindingPtrCustomization.h
new file mode 100644
--- /dev/null
+++ b/Engine/Source/Editor/DetailCustomizations/Private/MovieSceneObjectBindingPtrCustomization.h
@@ -0,0 +1,46 @@
+// Copyright 1998-2016 Epic Games, Inc. All Rights Reserved.
+
+#pragma once
+
+#include "IPropertyTypeCustomization.h"
+
+
+class IDetailsView;
+class IMovieSceneBindingOwnerInterface;
+class IPropertyHandle;
+class SWidget;
+class UMovieSceneSequence;
+
+struct FMovieScenePossessable;
+
+
+class FMovieSceneObjectBindingPtrCustomization
+	: public IPropertyTypeCustomization
+{
+public:
+
+	static TSharedRef<IPropertyTypeCustomization> MakeInstance()
+	{
+		return MakeShareable(new FMovieSceneObjectBindingPtrCustomization);
+	}
+
+	virtual void CustomizeHeader(TSharedRef<IPropertyHandle> PropertyHandle, FDetailWidgetRow& HeaderRow, IPropertyTypeCustomizationUtils& CustomizationUtils) override;
+	virtual void CustomizeChildren(TSharedRef<IPropertyHandle> StructPropertyHandle, class IDetailChildrenBuilder& ChildBuilder, IPropertyTypeCustomizationUtils& StructCustomizationUtils) override {}
+
+private:
+
+	FText ResolveBindingNameInSequence() const;
+
+	IMovieSceneBindingOwnerInterface* GetInterface() const;
+	UMovieSceneSequence* GetSequence() const;
+
+	void SetBindingId(FGuid InGuid);
+	FGuid GetGuid() const;
+
+	TSharedRef<SWidget> GetBindingMenu();
+
+private:
+
+	TSharedPtr<IPropertyHandle> StructProperty;
+	TSharedPtr<IPropertyHandle> GuidProperty;
+};
\ No newline at end of file
diff --git a/Engine/Source/Editor/DetailCustomizations/Private/MovieSceneSequencePlaybackSettingsCustomization.cpp b/Engine/Source/Editor/DetailCustomizations/Private/MovieSceneSequencePlaybackSettingsCustomization.cpp
new file mode 100644
--- /dev/null
+++ b/Engine/Source/Editor/DetailCustomizations/Private/MovieSceneSequencePlaybackSettingsCustomization.cpp
@@ -0,0 +1,183 @@
+// Copyright 1998-2017 Epic Games, Inc. All Rights Reserved.
+
+#include "MovieSceneSequencePlaybackSettingsCustomization.h"
+#include "Widgets/DeclarativeSyntaxSupport.h"
+#include "Engine/GameViewportClient.h"
+#include "Widgets/SBoxPanel.h"
+#include "Widgets/Images/SImage.h"
+#include "Widgets/Text/STextBlock.h"
+#include "Widgets/Input/SButton.h"
+#include "Widgets/Input/SComboBox.h"
+#include "PropertyHandle.h"
+#include "DetailLayoutBuilder.h"
+#include "DetailWidgetRow.h"
+#include "MovieSceneSequencePlayer.h"
+#include "IDetailChildrenBuilder.h"
+
+
+#define LOCTEXT_NAMESPACE "MovieSceneSequencePlaybackSettingsCustomization"
+
+
+TSharedRef<IPropertyTypeCustomization> FMovieSceneSequencePlaybackSettingsCustomization::MakeInstance()
+{
+	return MakeShareable(new FMovieSceneSequencePlaybackSettingsCustomization);
+}
+
+
+void FMovieSceneSequencePlaybackSettingsCustomization::CustomizeHeader(TSharedRef<IPropertyHandle> InPropertyHandle, FDetailWidgetRow& HeaderRow, IPropertyTypeCustomizationUtils& CustomizationUtils)
+{
+	// do nothing
+}
+
+
+void FMovieSceneSequencePlaybackSettingsCustomization::CustomizeChildren(TSharedRef<IPropertyHandle> InPropertyHandle, IDetailChildrenBuilder& ChildBuilder, IPropertyTypeCustomizationUtils& CustomizationUtils)
+{
+	StructPropertyHandle = InPropertyHandle;
+	ChildBuilder.AddChildProperty(StructPropertyHandle->GetChildHandle(GET_MEMBER_NAME_CHECKED(FMovieSceneSequencePlaybackSettings, PlayRate)).ToSharedRef());
+	LoopCountProperty = StructPropertyHandle->GetChildHandle(GET_MEMBER_NAME_CHECKED(FMovieSceneSequencePlaybackSettings, LoopCount));
+
+	// Set up the initial environment
+	{
+		LoopModes.Add(MakeShareable( new FLoopMode{ LOCTEXT("DontLoop", "Don't Loop"), 0 } ));
+		LoopModes.Add(MakeShareable( new FLoopMode{ LOCTEXT("Indefinitely", "Loop Indefinitely"), -1 }) );
+		LoopModes.Add(MakeShareable( new FLoopMode{ LOCTEXT("Exactly", "Loop Exactly..."), 1 }) );
+
+		int32 CurrentValue = -1;
+		LoopCountProperty->GetValue(CurrentValue);
+
+		TSharedPtr<FLoopMode>* FoundMode = LoopModes.FindByPredicate([&](const TSharedPtr<FLoopMode>& In){ return In->Value == CurrentValue; });
+		if (FoundMode)
+		{
+			CurrentMode = *FoundMode;
+		}
+		else
+		{
+			CurrentMode = LoopModes.Last();
+		}
+	}
+
+	ChildBuilder.AddChildContent(LOCTEXT("LoopTitle", "Loop"))
+		.NameContent()
+		[
+			LoopCountProperty->CreatePropertyNameWidget()
+		]
+		.ValueContent()
+		.HAlign(HAlign_Fill)
+		.MaxDesiredWidth(200)
+		[
+			SNew(SHorizontalBox)
+
+			+ SHorizontalBox::Slot()
+				.VAlign(VAlign_Center)
+				.AutoWidth()
+				[
+					SNew(SComboBox<TSharedPtr<FLoopMode>>)
+						.OptionsSource(&LoopModes)
+						.OnSelectionChanged_Lambda([&](TSharedPtr<FLoopMode> Mode, ESelectInfo::Type){
+							CurrentMode = Mode;
+							UpdateProperty();
+						})
+						.OnGenerateWidget_Lambda([&](TSharedPtr<FLoopMode> InMode){
+							return SNew(STextBlock)
+								.Font(IDetailLayoutBuilder::GetDetailFont())
+								.Text(InMode->DisplayName);
+						})
+						.InitiallySelectedItem(CurrentMode)
+						[
+							SAssignNew(CurrentText, STextBlock)
+								.Font(IDetailLayoutBuilder::GetDetailFont())
+								.Text(CurrentMode->DisplayName)
+						]
+				]
+
+			+ SHorizontalBox::Slot()
+				.Padding(FMargin(4.f,0))
+				.VAlign(VAlign_Center)
+				[
+					SAssignNew(LoopEntry, SHorizontalBox)
+						.Visibility(CurrentMode == LoopModes.Last() ? EVisibility::Visible : EVisibility::Collapsed)
+
+					+ SHorizontalBox::Slot()
+						.VAlign(VAlign_Center)
+						.Padding(FMargin(0, 0, 4.f, 0))
+						[
+							LoopCountProperty->CreatePropertyValueWidget()
+						]
+
+					+ SHorizontalBox::Slot()
+						.VAlign(VAlign_Center)
+						.AutoWidth()
+						[
+							SNew(STextBlock)
+								.Font(IDetailLayoutBuilder::GetDetailFont())
+								.Text_Raw(this, &FMovieSceneSequencePlaybackSettingsCustomization::GetCustomSuffix)
+						]
+				]
+
+			+SHorizontalBox::Slot()
+				.VAlign(VAlign_Center)
+				.AutoWidth()
+				[
+					SNew(SButton)
+					.OnClicked(this, &FMovieSceneSequencePlaybackSettingsCustomization::OnLoopResetClicked)
+					.Visibility(this, &FMovieSceneSequencePlaybackSettingsCustomization::GetLoopResetVisibility)
+					.ContentPadding(FMargin(5.f, 0.f))
+					.ToolTipText(LOCTEXT("ResetToDefaultToolTip", "Reset to Default"))
+					.ButtonStyle( FEditorStyle::Get(), "NoBorder" )
+					.Content()
+					[
+						SNew(SImage)
+						.Image( FEditorStyle::GetBrush("PropertyWindow.DiffersFromDefault") )
+					]
+				]
+		];
+
+	ChildBuilder.AddChildProperty(StructPropertyHandle->GetChildHandle(GET_MEMBER_NAME_CHECKED(FMovieSceneSequencePlaybackSettings, bRandomStartTime)).ToSharedRef());
+	ChildBuilder.AddChildProperty(StructPropertyHandle->GetChildHandle(GET_MEMBER_NAME_CHECKED(FMovieSceneSequencePlaybackSettings, StartTime)).ToSharedRef());
+	ChildBuilder.AddChildProperty(StructPropertyHandle->GetChildHandle(GET_MEMBER_NAME_CHECKED(FMovieSceneSequencePlaybackSettings, bRestoreState)).ToSharedRef());
+}
+
+
+FText FMovieSceneSequencePlaybackSettingsCustomization::GetCustomSuffix() const
+{
+	int32 NumLoops = 0;
+	LoopCountProperty->GetValue(NumLoops);
+
+	return NumLoops == 1
+		? LOCTEXT("Time", "time")
+		: LOCTEXT("Times", "times");
+}
+
+EVisibility FMovieSceneSequencePlaybackSettingsCustomization::GetLoopResetVisibility() const
+{
+	return CurrentMode == LoopModes[0] ? EVisibility::Hidden : EVisibility::Visible;
+}
+
+FReply FMovieSceneSequencePlaybackSettingsCustomization::OnLoopResetClicked()
+{
+	CurrentMode = LoopModes[0];
+
+	UpdateProperty();
+
+	return FReply::Handled();
+}
+
+void FMovieSceneSequencePlaybackSettingsCustomization::UpdateProperty()
+{
+	if (CurrentMode == LoopModes.Last())
+	{
+		// Show custom loop entry
+		LoopEntry->SetVisibility(EVisibility::Visible);
+	}
+	else
+	{
+		// Hide custom loop entry
+		LoopEntry->SetVisibility(EVisibility::Collapsed);
+	}
+
+	LoopCountProperty->SetValue(CurrentMode->Value);
+	CurrentText->SetText(CurrentMode->DisplayName);
+}
+
+
+#undef LOCTEXT_NAMESPACE
diff --git a/Engine/Source/Editor/DetailCustomizations/Private/MovieSceneSequencePlaybackSettingsCustomization.h b/Engine/Source/Editor/DetailCustomizations/Private/MovieSceneSequencePlaybackSettingsCustomization.h
new file mode 100644
--- /dev/null
+++ b/Engine/Source/Editor/DetailCustomizations/Private/MovieSceneSequencePlaybackSettingsCustomization.h
@@ -0,0 +1,69 @@
+// Copyright 1998-2017 Epic Games, Inc. All Rights Reserved.
+
+#pragma once
+
+#include "CoreMinimal.h"
+#include "Layout/Visibility.h"
+#include "Input/Reply.h"
+#include "IPropertyTypeCustomization.h"
+
+class FDetailWidgetRow;
+class IDetailChildrenBuilder;
+class IPropertyHandle;
+class STextBlock;
+class SWidget;
+
+/**
+ * Details view customization for the FMovieSceneSequencePlaybackSettings struct.
+ */
+class FMovieSceneSequencePlaybackSettingsCustomization
+	: public IPropertyTypeCustomization
+{
+public:
+
+	/** Makes a new instance of this detail layout class for a specific detail view requesting it. */
+	static TSharedRef<IPropertyTypeCustomization> MakeInstance();
+
+public:
+
+	// IDetailCustomization interface
+
+	virtual void CustomizeHeader( TSharedRef<IPropertyHandle> PropertyHandle, FDetailWidgetRow& HeaderRow, IPropertyTypeCustomizationUtils& CustomizationUtils ) override;
+	virtual void CustomizeChildren( TSharedRef<IPropertyHandle> PropertyHandle, IDetailChildrenBuilder& ChildBuilder, IPropertyTypeCustomizationUtils& CustomizationUtils ) override;
+
+protected:
+
+	/** Respond to a selection change event from the combo box. */
+	void UpdateProperty();
+	
+private:
+
+	/** Get the suffix to display after the custom entry box. */
+	FText GetCustomSuffix() const;
+
+	/** Get whether the loop reset button should be visible. */
+	EVisibility GetLoopResetVisibility() const;
+
+	/** Called when the "Reset to Default" button for the loop mode has been clicked. */
+	FReply OnLoopResetClicked();
+
+	/** Array of loop modes. */
+	struct FLoopMode
+	{
+		FText DisplayName;
+		int32 Value;
+	};
+
+	TArray<TSharedPtr<FLoopMode>> LoopModes;
+
+	/** The loop mode we're currently displaying. */
+	TSharedPtr<FLoopMode> CurrentMode;
+	/** The text of the current selection. */
+	TSharedPtr<STextBlock> CurrentText;
+	/** The loop number entry to be hidden and shown based on combo box selection. */
+	TSharedPtr<SWidget> LoopEntry;
+
+	/** Property handles of the properties we're editing */
+	TSharedPtr<IPropertyHandle> StructPropertyHandle;
+	TSharedPtr<IPropertyHandle> LoopCountProperty;
+};
diff --git a/Engine/Source/Editor/EditorSettingsViewer/Private/EditorSettingsViewerModule.cpp b/Engine/Source/Editor/EditorSettingsViewer/Private/EditorSettingsViewerModule.cpp
--- a/Engine/Source/Editor/EditorSettingsViewer/Private/EditorSettingsViewerModule.cpp
+++ b/Engine/Source/Editor/EditorSettingsViewer/Private/EditorSettingsViewerModule.cpp
@@ -37,6 +37,7 @@
 
 #include "Tests/AutomationTestSettings.h"
 #include "BlueprintEditorSettings.h"
+#include "CurveEditorSettings.h"
 
 #include "CrashReporterSettings.h"
 #include "Analytics/AnalyticsPrivacySettings.h"
@@ -233,6 +234,13 @@ protected:
 			LOCTEXT("ContentEditorsPersonaSettingsDescription", "Customize Persona Editor."),
 			GetMutableDefault<UPersonaOptions>()
 		);
+
+		// curve editor
+		SettingsModule.RegisterSettings("Editor", "ContentEditors", "CurveEditor",
+			LOCTEXT("ContentEditorsCurveEditorSettingsName", "Curve Editor"),
+			LOCTEXT("ContentEditorsCurveEditorSettingsDescription", "Customize Curve Editors."),
+			GetMutableDefault<UCurveEditorSettings>()
+		);
 	}
 
 	void RegisterPrivacySettings(ISettingsModule& SettingsModule)
@@ -273,6 +281,7 @@ protected:
 //			SettingsModule->UnregisterSettings("Editor", "ContentEditors", "DestructableMeshEditor");
 			SettingsModule->UnregisterSettings("Editor", "ContentEditors", "GraphEditor");
 			SettingsModule->UnregisterSettings("Editor", "ContentEditors", "Persona");
+			SettingsModule->UnregisterSettings("Editor", "ContentEditors", "CurveEditor");
 		}
 	}
 
diff --git a/Engine/Source/Editor/EditorStyle/Private/SlateEditorStyle.cpp b/Engine/Source/Editor/EditorStyle/Private/SlateEditorStyle.cpp
--- a/Engine/Source/Editor/EditorStyle/Private/SlateEditorStyle.cpp
+++ b/Engine/Source/Editor/EditorStyle/Private/SlateEditorStyle.cpp
@@ -1755,6 +1755,8 @@ void FSlateEditorStyle::FStyle::SetupGeneralStyles()
 		Set( "Sequencer.SaveAs.Small", new IMAGE_BRUSH( "Sequencer/Main_Icons/Icon_Sequencer_Save_As_24x", Icon24x24 ) );
 		Set( "Sequencer.DiscardChanges", new IMAGE_BRUSH( "Sequencer/Main_Icons/Icon_Sequencer_Revert_24x", Icon48x48 ) );
 		Set( "Sequencer.DiscardChanges.Small", new IMAGE_BRUSH( "Sequencer/Main_Icons/Icon_Sequencer_Revert_24x", Icon24x24 ) );
+		Set( "Sequencer.RestoreAnimatedState", new IMAGE_BRUSH("Sequencer/Main_Icons/Icon_Sequencer_RestoreAnimatedState_24x", Icon48x48 ) );
+		Set( "Sequencer.RestoreAnimatedState.Small", new IMAGE_BRUSH("Sequencer/Main_Icons/Icon_Sequencer_RestoreAnimatedState_24x", Icon24x24 ) );
 		Set( "Sequencer.GenericGripLeft", new BOX_BRUSH( "Sequencer/GenericGripLeft", FMargin(5.f/16.f) ) );
 		Set( "Sequencer.GenericGripRight", new BOX_BRUSH( "Sequencer/GenericGripRight", FMargin(5.f/16.f) ) );
 		Set( "Sequencer.SectionArea.Background", new FSlateColorBrush( FColor::White ) );
@@ -1816,8 +1818,12 @@ void FSlateEditorStyle::FStyle::SetupGeneralStyles()
 		Set( "Sequencer.TrackHoverHighlight_Bottom", new IMAGE_BRUSH( TEXT("Sequencer/TrackHoverHighlight_Bottom"), FVector2D(4, 4) ) );
 		Set( "Sequencer.SpawnableIconOverlay", new IMAGE_BRUSH( TEXT("Sequencer/SpawnableIconOverlay"), FVector2D(13, 13) ) );
 
-		Set( "Sequencer.General", new IMAGE_BRUSH( "Sequencer/Main_Icons/Icon_Sequencer_General_Options_24x", Icon48x48 ) );
-		Set( "Sequencer.General.Small", new IMAGE_BRUSH( "Sequencer/Main_Icons/Icon_Sequencer_General_Options_24x", Icon24x24 ) );
+		Set( "Sequencer.GeneralOptions", new IMAGE_BRUSH( "Sequencer/Main_Icons/Icon_Sequencer_General_Options_24x", Icon48x48 ) );
+		Set( "Sequencer.GeneralOptions.Small", new IMAGE_BRUSH( "Sequencer/Main_Icons/Icon_Sequencer_General_Options_24x", Icon24x24 ) );
+		Set( "Sequencer.PlaybackOptions", new IMAGE_BRUSH( "Sequencer/Main_Icons/Icon_Sequencer_Playback_Options_24x", Icon48x48 ) );
+		Set( "Sequencer.PlaybackOptions.Small", new IMAGE_BRUSH( "Sequencer/Main_Icons/Icon_Sequencer_Playback_Options_24x", Icon24x24 ) );
+		Set( "Sequencer.SelectEditOptions", new IMAGE_BRUSH( "Sequencer/Main_Icons/Icon_Sequencer_SelectEdit_Options_24x", Icon48x48 ) );
+		Set( "Sequencer.SelectEditOptions.Small", new IMAGE_BRUSH( "Sequencer/Main_Icons/Icon_Sequencer_SelectEdit_Options_24x", Icon24x24 ) );
 		Set( "Sequencer.Time", new IMAGE_BRUSH( "Sequencer/Main_Icons/Icon_Sequencer_Time_24x", Icon48x48 ) );
 		Set( "Sequencer.Time.Small", new IMAGE_BRUSH( "Sequencer/Main_Icons/Icon_Sequencer_Time_24x", Icon24x24 ) );
 		Set( "Sequencer.Value", new IMAGE_BRUSH( "Sequencer/Main_Icons/Icon_Sequencer_Value_24x", Icon48x48 ) );
@@ -4291,6 +4297,8 @@ void FSlateEditorStyle::FStyle::SetupGraphEditorStyles()
 		Set( "GraphEditor.FunctionOL.Override", new IMAGE_BRUSH( "Graph/Icons/Overlay_Override", Icon22x22 ) );
 		Set( "GraphEditor.FunctionOL.PotentialOverride", new IMAGE_BRUSH( "Graph/Icons/Overlay_PotentialOverride", Icon22x22 ) );
 
+		Set( "GraphEditor.GetSequenceBindings", new IMAGE_BRUSH("Icons/icon_Blueprint_GetSequenceBindings_16x", Icon16x16));
+
 		Set( "GraphEditor.HideUnusedPins", new IMAGE_BRUSH( "Icons/hide_unusedpins", Icon40x40 ) );
 		Set( "GraphEditor.HideUnusedPins.Small", new IMAGE_BRUSH( "Icons/hide_unusedpins", Icon20x20 ) );
 
diff --git a/Engine/Source/Editor/Kismet/Private/BlueprintEditor.cpp b/Engine/Source/Editor/Kismet/Private/BlueprintEditor.cpp
--- a/Engine/Source/Editor/Kismet/Private/BlueprintEditor.cpp
+++ b/Engine/Source/Editor/Kismet/Private/BlueprintEditor.cpp
@@ -93,6 +93,7 @@
 #include "Kismet2/Kismet2NameValidators.h"
 #include "Kismet2/DebuggerCommands.h"
 #include "Editor.h"
+#include "IDetailsView.h"
 
 #include "BlueprintEditorTabs.h"
 
@@ -184,6 +185,8 @@ TSharedRef<SDockTab> FSelectionDetailsSummoner::SpawnTab(const FWorkflowTabSpawn
 	TSharedPtr<FBlueprintEditor> BlueprintEditorPtr = StaticCastSharedPtr<FBlueprintEditor>(HostingApp.Pin());
 	BlueprintEditorPtr->GetInspector()->SetOwnerTab(Tab);
 
+	BlueprintEditorPtr->GetInspector()->GetPropertyView()->SetHostTabManager(Info.TabManager);
+
 	return Tab;
 }
 
diff --git a/Engine/Source/Editor/Kismet/Private/BlueprintEditorModes.cpp b/Engine/Source/Editor/Kismet/Private/BlueprintEditorModes.cpp
--- a/Engine/Source/Editor/Kismet/Private/BlueprintEditorModes.cpp
+++ b/Engine/Source/Editor/Kismet/Private/BlueprintEditorModes.cpp
@@ -14,6 +14,7 @@
 // Debugging
 // End of debugging
 
+#include "LayoutExtender.h"
 #include "SBlueprintEditorToolbar.h"
 #include "BlueprintEditorTabs.h"
 #include "BlueprintEditorTabFactories.h"
@@ -121,6 +122,12 @@ FBlueprintEditorApplicationMode::FBlueprintEditorApplicationMode(TSharedPtr<clas
 	InBlueprintEditor->GetToolbarBuilder()->AddBlueprintGlobalOptionsToolbar(ToolbarExtender);
 	InBlueprintEditor->GetToolbarBuilder()->AddDebuggingToolbar(ToolbarExtender);
 	InBlueprintEditor->GetToolbarBuilder()->AddProfilerToolbar(ToolbarExtender);
+
+	FBlueprintEditorModule& BlueprintEditorModule = FModuleManager::LoadModuleChecked<FBlueprintEditorModule>("Kismet");
+	BlueprintEditorModule.OnRegisterTabsForEditor().Broadcast(BlueprintEditorTabFactories, InModeName, InBlueprintEditor);
+
+	LayoutExtender = MakeShared<FLayoutExtender>();
+	BlueprintEditorModule.OnRegisterLayoutExtensions().Broadcast(*LayoutExtender);
 }
 
 void FBlueprintEditorApplicationMode::RegisterTabFactories(TSharedPtr<FTabManager> InTabManager)
@@ -574,7 +581,7 @@ FBlueprintEditorUnifiedMode::FBlueprintEditorUnifiedMode(TSharedPtr<class FBluep
 
 	if ( bRegisterViewport )
 	{
-		TabLayout = FTabManager::NewLayout( "Blueprints_Unified_Components_v5" )
+		TabLayout = FTabManager::NewLayout( "Blueprints_Unified_Components_v6" )
 		->AddArea
 		(
 			FTabManager::NewPrimaryArea()->SetOrientation(Orient_Vertical)
@@ -719,6 +726,12 @@ FBlueprintEditorUnifiedMode::FBlueprintEditorUnifiedMode(TSharedPtr<class FBluep
 	}
 	
 	InBlueprintEditor->GetToolbarBuilder()->AddDebuggingToolbar(ToolbarExtender);
+
+	FBlueprintEditorModule& BlueprintEditorModule = FModuleManager::LoadModuleChecked<FBlueprintEditorModule>("Kismet");
+	BlueprintEditorModule.OnRegisterTabsForEditor().Broadcast(BlueprintEditorTabFactories, InModeName, InBlueprintEditor);
+
+	LayoutExtender = MakeShared<FLayoutExtender>();
+	BlueprintEditorModule.OnRegisterLayoutExtensions().Broadcast(*LayoutExtender);
 }
 
 void FBlueprintEditorUnifiedMode::RegisterTabFactories(TSharedPtr<FTabManager> InTabManager)
diff --git a/Engine/Source/Editor/Kismet/Private/SSCSEditor.cpp b/Engine/Source/Editor/Kismet/Private/SSCSEditor.cpp
--- a/Engine/Source/Editor/Kismet/Private/SSCSEditor.cpp
+++ b/Engine/Source/Editor/Kismet/Private/SSCSEditor.cpp
@@ -1777,20 +1777,24 @@ EVisibility SSCS_RowWidget::GetAssetVisibility() const
 }
 
 FSlateColor SSCS_RowWidget::GetColorTintForIcon() const
+{
+	return GetColorTintForIcon(GetNode());
+}
+
+FSlateColor SSCS_RowWidget::GetColorTintForIcon(FSCSEditorTreeNodePtrType InNode)
 {
 	const FLinearColor InheritedBlueprintComponentColor(0.08f, 0.35f, 0.6f);
 	const FLinearColor InstancedInheritedBlueprintComponentColor(0.08f, 0.35f, 0.6f);
 	const FLinearColor InheritedNativeComponentColor(0.7f, 0.9f, 0.7f);
 	const FLinearColor IntroducedHereColor(FLinearColor::White);
-
-	FSCSEditorTreeNodePtrType NodePtr = GetNode();
-	if (NodePtr->IsInherited())
+	
+	if (InNode->IsInherited())
 	{
-		if (NodePtr->IsNative())
+		if (InNode->IsNative())
 		{
 			return InheritedNativeComponentColor;
 		}
-		else if (NodePtr->IsInstanced())
+		else if (InNode->IsInstanced())
 		{
 			return InstancedInheritedBlueprintComponentColor;
 		}
diff --git a/Engine/Source/Editor/Kismet/Private/WorkflowCentricApplication.cpp b/Engine/Source/Editor/Kismet/Private/WorkflowCentricApplication.cpp
--- a/Engine/Source/Editor/Kismet/Private/WorkflowCentricApplication.cpp
+++ b/Engine/Source/Editor/Kismet/Private/WorkflowCentricApplication.cpp
@@ -42,8 +42,15 @@ void FWorkflowCentricApplication::SetCurrentMode(FName NewMode)
 
 		TSharedPtr<FApplicationMode> NewModePtr = ApplicationModeList.FindRef(NewMode);
 
+		LayoutExtenders.Reset();
+
 		if (NewModePtr.IsValid())
 		{
+			if (NewModePtr->LayoutExtender.IsValid())
+			{
+				LayoutExtenders.Add(NewModePtr->LayoutExtender);
+			}
+			
 			// Deactivate the old mode
 			if (CurrentAppModePtr.IsValid())
 			{
@@ -77,33 +84,13 @@ void FWorkflowCentricApplication::SetCurrentMode(FName NewMode)
 	}
 }
 
-
-TSharedRef<SDockTab> FWorkflowCentricApplication::CreatePanelTab(const FSpawnTabArgs& Args, TSharedPtr<FWorkflowTabFactory> TabFactory)
-{
-	FWorkflowTabSpawnInfo SpawnInfo;
-
-	return TabFactory->SpawnTab(SpawnInfo);
-}
-
 void FWorkflowCentricApplication::PushTabFactories(FWorkflowAllowedTabSet& FactorySetToPush)
 {
-	FWorkflowTabSpawnInfo SpawnInfo;
-
+	check(TabManager.IsValid());
 	for (auto FactoryIt = FactorySetToPush.CreateIterator(); FactoryIt; ++FactoryIt)
 	{
-		TSharedPtr<FWorkflowTabFactory> SomeFactory = FactoryIt.Value();
-		FTabSpawnerEntry& SpawnerEntry = TabManager->RegisterTabSpawner(SomeFactory->GetIdentifier(), FOnSpawnTab::CreateRaw(this, &FWorkflowCentricApplication::CreatePanelTab, SomeFactory))
-			.SetDisplayName(SomeFactory->ConstructTabName(SpawnInfo).Get())
-			.SetTooltipText(SomeFactory->GetTabToolTipText(SpawnInfo))
-			.SetGroup(CurrentAppModePtr->GetWorkspaceMenuCategory());
-		
-		// Add the tab icon to the menu entry if one was provided
-		const FSlateIcon& TabSpawnerIcon = SomeFactory->GetTabSpawnerIcon(SpawnInfo);
-		if (TabSpawnerIcon.IsSet())
-		{
-			SpawnerEntry.SetIcon(TabSpawnerIcon);
-		}
-	} 
+		FactoryIt.Value()->RegisterTabSpawner(TabManager.ToSharedRef(), CurrentAppModePtr.Get());
+	}
 }
 
 bool FWorkflowCentricApplication::OnRequestClose()
diff --git a/Engine/Source/Editor/Kismet/Public/BlueprintEditorModule.h b/Engine/Source/Editor/Kismet/Public/BlueprintEditorModule.h
--- a/Engine/Source/Editor/Kismet/Public/BlueprintEditorModule.h
+++ b/Engine/Source/Editor/Kismet/Public/BlueprintEditorModule.h
@@ -18,6 +18,8 @@ extern const FName BlueprintEditorAppName;
 
 class FBlueprintEditor;
 class UUserDefinedEnum;
+class UUserDefinedStruct;
+
 
 /** Describes the reason for Refreshing the editor */
 namespace ERefreshBlueprintEditorReason
@@ -142,6 +144,12 @@ public:
 	DECLARE_EVENT_TwoParams(FBlueprintEditorModule, FBlueprintMenuExtensionEvent, TSharedPtr<FExtender>, UBlueprint*);
 	FBlueprintMenuExtensionEvent& OnGatherBlueprintMenuExtensions() { return GatherBlueprintMenuExtensions; }
 
+	DECLARE_EVENT_ThreeParams(IBlueprintEditor, FOnRegisterTabs, FWorkflowAllowedTabSet&, FName /** ModeName */, TSharedPtr<FBlueprintEditor>);
+	FOnRegisterTabs& OnRegisterTabsForEditor() { return RegisterTabsForEditor; }
+
+	DECLARE_EVENT_OneParam(IBlueprintEditor, FOnRegisterLayoutExtensions, FLayoutExtender&);
+	FOnRegisterLayoutExtensions& OnRegisterLayoutExtensions() { return RegisterLayoutExtensions; }
+
 	/** 
 	 * Register a customization for interacting with the SCS editor 
 	 * @param	InComponentName			The name of the component to customize behavior for
@@ -176,6 +184,10 @@ private:
 	//
 	FBlueprintMenuExtensionEvent GatherBlueprintMenuExtensions;
 
+	/** Event called to allow external clients to register additional tabs for the specified editor */
+	FOnRegisterTabs RegisterTabsForEditor;
+	FOnRegisterLayoutExtensions RegisterLayoutExtensions;
+
 	// Event to be called when the blueprint editor is opened
 	FBlueprintEditorOpenedEvent BlueprintEditorOpened;
 
diff --git a/Engine/Source/Editor/Kismet/Public/SSCSEditor.h b/Engine/Source/Editor/Kismet/Public/SSCSEditor.h
--- a/Engine/Source/Editor/Kismet/Public/SSCSEditor.h
+++ b/Engine/Source/Editor/Kismet/Public/SSCSEditor.h
@@ -521,6 +521,8 @@ public:
 	FString GetDocumentationLink() const;
 	FString GetDocumentationExcerptName() const;
 	
+	KISMET_API static FSlateColor GetColorTintForIcon(FSCSEditorTreeNodePtrType InNode);
+
 	FText GetAssetName() const;
 	FText GetAssetPath() const;
 	EVisibility GetAssetVisibility() const;
diff --git a/Engine/Source/Editor/Kismet/Public/WorkflowOrientedApp/ApplicationMode.h b/Engine/Source/Editor/Kismet/Public/WorkflowOrientedApp/ApplicationMode.h
--- a/Engine/Source/Editor/Kismet/Public/WorkflowOrientedApp/ApplicationMode.h
+++ b/Engine/Source/Editor/Kismet/Public/WorkflowOrientedApp/ApplicationMode.h
@@ -5,9 +5,12 @@
 #include "CoreMinimal.h"
 #include "Framework/Docking/WorkspaceItem.h"
 #include "Framework/Docking/TabManager.h"
+#include "LayoutExtender.h"
+
 
 class FExtender;
 
+
 /////////////////////////////////////////////////////
 // FApplicationMode
 
@@ -55,5 +58,8 @@ public:
 	TSharedPtr<FExtender> GetToolbarExtender() { return ToolbarExtender; }
 
 	/** @return The the workspace category for this asset editor */
-	TSharedRef<FWorkspaceItem> GetWorkspaceMenuCategory() { return WorkspaceMenuCategory.ToSharedRef(); }
+	TSharedRef<FWorkspaceItem> GetWorkspaceMenuCategory() const { return WorkspaceMenuCategory.ToSharedRef(); }
+
+	/** Extender for adding to the default layout for this mode */
+	TSharedPtr<FLayoutExtender> LayoutExtender;
 };
diff --git a/Engine/Source/Editor/Kismet/Public/WorkflowOrientedApp/WorkflowCentricApplication.h b/Engine/Source/Editor/Kismet/Public/WorkflowOrientedApp/WorkflowCentricApplication.h
--- a/Engine/Source/Editor/Kismet/Public/WorkflowOrientedApp/WorkflowCentricApplication.h
+++ b/Engine/Source/Editor/Kismet/Public/WorkflowOrientedApp/WorkflowCentricApplication.h
@@ -40,7 +40,6 @@ public:
 	// Gets the mode extender list for all workflow applications (append to customize a specific mode)
 	static TArray<FWorkflowApplicationModeExtender>& GetModeExtenderList() { return ModeExtenderList; }
 protected:
-	TSharedRef<SDockTab> CreatePanelTab(const FSpawnTabArgs& Args, TSharedPtr<FWorkflowTabFactory> TabFactory);
 
 	virtual void AddApplicationMode(FName ModeName, TSharedRef<FApplicationMode> Mode);
 protected:
diff --git a/Engine/Source/Editor/Kismet/Public/WorkflowOrientedApp/WorkflowTabFactory.h b/Engine/Source/Editor/Kismet/Public/WorkflowOrientedApp/WorkflowTabFactory.h
--- a/Engine/Source/Editor/Kismet/Public/WorkflowOrientedApp/WorkflowTabFactory.h
+++ b/Engine/Source/Editor/Kismet/Public/WorkflowOrientedApp/WorkflowTabFactory.h
@@ -8,6 +8,7 @@
 #include "Textures/SlateIcon.h"
 #include "Toolkits/AssetEditorToolkit.h"
 
+class FApplicationMode;
 class FMenuBuilder;
 class FTabInfo;
 class SToolTip;
@@ -54,6 +55,9 @@ public:
 	// The TabInfo being used to spawn this
 	TSharedPtr<class FTabInfo> TabInfo;
 
+	/** The parent tab manager */
+	TSharedPtr<FTabManager> TabManager;
+
 public:
 	FWorkflowTabSpawnInfo()
 		: TabInfo(NULL)
@@ -103,6 +107,9 @@ public:
 	// Creates a view menu entry for this factory
 	virtual void CreateViewMenuEntry(FMenuBuilder& MenuBuilder, const FUIAction& Action) const;
 
+	/** Register a tab spawner for this factory */
+	virtual FTabSpawnerEntry& RegisterTabSpawner(TSharedRef<FTabManager> TabManager, const FApplicationMode* CurrentApplicationMode) const;
+
 public:
 	// Constructor
 	FWorkflowTabFactory(FName InIdentifier, TSharedPtr<class FAssetEditorToolkit> InHostingApp);
@@ -123,6 +130,9 @@ protected:
 	// Call this to enable the default padding surrounding the tab contents
 	void EnableTabPadding();
 
+	/** Callback function for spawning the tab */
+	TSharedRef<SDockTab> OnSpawnTab(const FSpawnTabArgs& SpawnArgs, TWeakPtr<FTabManager> WeakTabManager) const;
+
 public:
 	// Creates the widget that will be the body of the new tab
 	// Users will most commonly override just this.
diff --git a/Engine/Source/Editor/Kismet/Public/WorkflowOrientedApp/WorkflowUObjectDocuments.h b/Engine/Source/Editor/Kismet/Public/WorkflowOrientedApp/WorkflowUObjectDocuments.h
--- a/Engine/Source/Editor/Kismet/Public/WorkflowOrientedApp/WorkflowUObjectDocuments.h
+++ b/Engine/Source/Editor/Kismet/Public/WorkflowOrientedApp/WorkflowUObjectDocuments.h
@@ -9,6 +9,8 @@
 #include "Widgets/SWidget.h"
 #include "Toolkits/AssetEditorToolkit.h"
 #include "WorkflowOrientedApp/WorkflowTabFactory.h"
+#include "WorkflowTabFactory.h"
+
 
 /////////////////////////////////////////////////////
 // FTabPayload_UObject
diff --git a/Engine/Source/Editor/KismetWidgets/Private/SSingleObjectDetailsPanel.cpp b/Engine/Source/Editor/KismetWidgets/Private/SSingleObjectDetailsPanel.cpp
--- a/Engine/Source/Editor/KismetWidgets/Private/SSingleObjectDetailsPanel.cpp
+++ b/Engine/Source/Editor/KismetWidgets/Private/SSingleObjectDetailsPanel.cpp
@@ -18,6 +18,7 @@ void SSingleObjectDetailsPanel::Construct(const FArguments& InArgs, bool bAutoma
 
 	FDetailsViewArgs DetailsViewArgs(/*bUpdateFromSelection=*/ false, /*bLockable=*/ false, bAllowSearch, FDetailsViewArgs::HideNameArea, /*bHideSelectionTip=*/ true);
 	DetailsViewArgs.HostCommandList = InArgs._HostCommandList;
+	DetailsViewArgs.HostTabManager = InArgs._HostTabManager;
 
 	PropertyView = EditModule.CreateDetailView(DetailsViewArgs);
 	
diff --git a/Engine/Source/Editor/KismetWidgets/Public/SSingleObjectDetailsPanel.h b/Engine/Source/Editor/KismetWidgets/Public/SSingleObjectDetailsPanel.h
--- a/Engine/Source/Editor/KismetWidgets/Public/SSingleObjectDetailsPanel.h
+++ b/Engine/Source/Editor/KismetWidgets/Public/SSingleObjectDetailsPanel.h
@@ -18,6 +18,7 @@ class KISMETWIDGETS_API SSingleObjectDetailsPanel : public SCompoundWidget
 public:
 	SLATE_BEGIN_ARGS(SSingleObjectDetailsPanel) {}
 		SLATE_ARGUMENT(TSharedPtr<class FUICommandList>, HostCommandList)
+		SLATE_ARGUMENT(TSharedPtr<class FTabManager>, HostTabManager)
 	SLATE_END_ARGS()
 
 	void Construct(const FArguments& InArgs, bool bAutomaticallyObserveViaGetObjectToObserve = true, bool bAllowSearch = false);
