Fixed Windows file picker potentially mangling the returned path

diff --git a/Engine/Source/Developer/DesktopPlatform/Private/Windows/DesktopPlatformWindows.cpp b/Engine/Source/Developer/DesktopPlatform/Private/Windows/DesktopPlatformWindows.cpp
--- a/Engine/Source/Developer/DesktopPlatform/Private/Windows/DesktopPlatformWindows.cpp
+++ b/Engine/Source/Developer/DesktopPlatform/Private/Windows/DesktopPlatformWindows.cpp
@@ -301,11 +301,17 @@ bool FDesktopPlatformWindows::FileDialogShared(bool bSave, const void* ParentWin
 								}
 							}
 
-							// Apply the extension if the file path doesn't already have one
-							if (FPaths::GetExtension(SaveFilePath).IsEmpty() && !CleanExtension.IsEmpty())
+							// We need to split these before testing the extension to avoid anything within the path being treated as a file extension
+							FString SaveFileName = FPaths::GetCleanFilename(SaveFilePath);
+							SaveFilePath = FPaths::GetPath(SaveFilePath);
+
+							// Apply the extension if the file name doesn't already have one
+							if (FPaths::GetExtension(SaveFileName).IsEmpty() && !CleanExtension.IsEmpty())
 							{
-								SaveFilePath = FPaths::SetExtension(SaveFilePath, CleanExtension);
+								SaveFileName = FPaths::SetExtension(SaveFileName, CleanExtension);
 							}
+
+							SaveFilePath /= SaveFileName;
 						}
 						AddOutFilename(SaveFilePath);
 
