Changing MLSDK missing callback to launch MagicLeap URL

diff --git a/Engine/Source/Developer/DesktopPlatform/Private/PlatformInfo.cpp b/Engine/Source/Developer/DesktopPlatform/Private/PlatformInfo.cpp
--- a/Engine/Source/Developer/DesktopPlatform/Private/PlatformInfo.cpp
+++ b/Engine/Source/Developer/DesktopPlatform/Private/PlatformInfo.cpp
@@ -134,7 +134,7 @@ static const FPlatformInfo AllPlatformInfoArray[] = {
 
 	BuildPlatformInfo(TEXT("TVOS"),						TEXT("TVOS"),				LOCTEXT("TVOSTargetPlatDisplay", "tvOS"),						EPlatformType::Game,	EPlatformFlags::None,			FPlatformIconPaths(TEXT("Launcher/TVOS/Platform_TVOS_24x"), TEXT("Launcher/TVOS/Platform_TVOS_128x")),								TEXT(""),											TEXT(""),			EPlatformSDKStatus::Unknown,	TEXT(""),																								IsAvailableOnWindows || IsAvailableOnMac,						TEXT("TVOS"),		TEXT("TVOS"),		false,					true,					false,			TEXT("TVOS"),		TEXT("Mobile")),
 	BuildPlatformInfo(TEXT("Switch"),					TEXT("Switch"),				LOCTEXT("Switch", "Switch"),									EPlatformType::Game,	EPlatformFlags::None,			FPlatformIconPaths(TEXT("Launcher/Switch/Platform_Switch_24x"), TEXT("Launcher/Switch/Platform_Switch_128x")),						TEXT(""),											TEXT(""),			EPlatformSDKStatus::Unknown,	TEXT(""),																								IsAvailableOnWindows,											TEXT("Switch"),		TEXT("Switch"),		false,					true,					true,			TEXT("Switch"),		TEXT("Console")),
-	BuildPlatformInfo(TEXT("Lumin"),					TEXT("Lumin"),				LOCTEXT("Lumin", "Lumin"),										EPlatformType::Game,	EPlatformFlags::None,			FPlatformIconPaths(TEXT("Launcher/LuminTarget/Platform_Lumin_24x"), TEXT("Launcher/Lumin/Platform_Lumin_128x")),					TEXT(""),											TEXT(""),			EPlatformSDKStatus::Unknown,	TEXT(""),																								IsAvailableOnLinux || IsAvailableOnWindows || IsAvailableOnMac,	TEXT("Lumin"),		TEXT("Lumin"),		false,					true,					false,			TEXT("Lumin"),		TEXT("Mobile")),
+	BuildPlatformInfo(TEXT("Lumin"),					TEXT("Lumin"),				LOCTEXT("Lumin", "Lumin"),										EPlatformType::Game,	EPlatformFlags::None,			FPlatformIconPaths(TEXT("Launcher/LuminTarget/Platform_Lumin_24x"), TEXT("Launcher/Lumin/Platform_Lumin_128x")),					TEXT(""),											TEXT(""),			EPlatformSDKStatus::Unknown,	TEXT("/Engine/Tutorial/Mobile/SettingUpLuminTutorial.SettingUpLuminTutorial"),							IsAvailableOnLinux || IsAvailableOnWindows || IsAvailableOnMac,	TEXT("Lumin"),		TEXT("Lumin"),		false,					true,					false,			TEXT("Lumin"),		TEXT("Mobile")),
 	BuildPlatformInfo(TEXT("Quail"),					TEXT("Quail"),				LOCTEXT("Quail", "Quail"),										EPlatformType::Game,	EPlatformFlags::None,			FPlatformIconPaths(TEXT("Launcher/Quail/Platform_Quail_24x"), TEXT("Launcher/Quail/Platform_Quail_128x")),							TEXT(""),											TEXT(""),			EPlatformSDKStatus::Unknown,	TEXT(""),																								IsAvailableOnWindows,											TEXT("Quail"),		TEXT("Quail"),		false,					true,					true,			TEXT("Quail"),		TEXT("")),
 
 
diff --git a/Engine/Source/Editor/UnrealEd/Private/Kismet2/DebuggerCommands.cpp b/Engine/Source/Editor/UnrealEd/Private/Kismet2/DebuggerCommands.cpp
--- a/Engine/Source/Editor/UnrealEd/Private/Kismet2/DebuggerCommands.cpp
+++ b/Engine/Source/Editor/UnrealEd/Private/Kismet2/DebuggerCommands.cpp
@@ -963,6 +963,7 @@ TSharedRef< SWidget > FPlayWorldCommands::GenerateLaunchMenuContent( TSharedRef<
 	PlatformsToMaybeInstallLinksFor.Add(TEXT("Android"));
 	PlatformsToMaybeInstallLinksFor.Add(TEXT("IOS"));
 	PlatformsToMaybeInstallLinksFor.Add(TEXT("Linux"));
+	PlatformsToMaybeInstallLinksFor.Add(TEXT("Lumin"));
 	TArray<FString> PlatformsToCheckFlavorsFor;
 	PlatformsToCheckFlavorsFor.Add(TEXT("Android"));
 	PlatformsToCheckFlavorsFor.Add(TEXT("IOS"));
