Speculative fix for crash in macOS AppKit NSPersistentUIManager.  Disable NSWindow application resume feature per window - we are not using it anyway.  OS system may still be called but hopefully this reduces the amount of work it is doing.  There still seems to be some generic data getting written however.

There is also an undocumented *NSDisablePersistence* plist key.  Using this plist key completely stops any data getting written by the OS window resume feature. This CL does NOT use this because it's undocumented but maybe an option for the future.

diff --git a/Engine/Source/Developer/DesktopPlatform/Private/Mac/MacNativeFeedbackContext.cpp b/Engine/Source/Developer/DesktopPlatform/Private/Mac/MacNativeFeedbackContext.cpp
--- a/Engine/Source/Developer/DesktopPlatform/Private/Mac/MacNativeFeedbackContext.cpp
+++ b/Engine/Source/Developer/DesktopPlatform/Private/Mac/MacNativeFeedbackContext.cpp
@@ -65,6 +65,8 @@
 		[Window setTitle:@"Unreal Engine 4"];
 		[Window setReleasedWhenClosed:NO];
 		[Window setMinSize:NSMakeSize(400, 100)];
+		[Window setRestorable:NO];
+		[Window disableSnapshotRestoration];
 		
 		NSView* View = [Window contentView];
 		[View setAutoresizesSubviews:YES];
diff --git a/Engine/Source/Programs/Mac/UE4EditorServices/Private/Mac/UE4EditorServicesAppDelegate.cpp b/Engine/Source/Programs/Mac/UE4EditorServices/Private/Mac/UE4EditorServicesAppDelegate.cpp
--- a/Engine/Source/Programs/Mac/UE4EditorServices/Private/Mac/UE4EditorServicesAppDelegate.cpp
+++ b/Engine/Source/Programs/Mac/UE4EditorServices/Private/Mac/UE4EditorServicesAppDelegate.cpp
@@ -162,6 +162,8 @@
 	[[Window contentView] addSubview:CancelButton];
 
 	[Window setTitle:@"Choose UE4 Editor..."];
+	[Window setRestorable:NO];
+	[Window disableSnapshotRestoration];
 
 	[Window center];
 	[NSApp activateIgnoringOtherApps:YES];
diff --git a/Engine/Source/Runtime/ApplicationCore/Private/Mac/CocoaWindow.cpp b/Engine/Source/Runtime/ApplicationCore/Private/Mac/CocoaWindow.cpp
--- a/Engine/Source/Runtime/ApplicationCore/Private/Mac/CocoaWindow.cpp
+++ b/Engine/Source/Runtime/ApplicationCore/Private/Mac/CocoaWindow.cpp
@@ -36,6 +36,8 @@ NSString* NSPerformDragOperation = @"NSPerformDragOperation";
 		bIsOnActiveSpace = [super isOnActiveSpace];
 		self.TargetWindowMode = EWindowMode::Windowed;
 		[super setAlphaValue:Opacity];
+		[super setRestorable:NO];
+		[super disableSnapshotRestoration];
 	}
 	return NewSelf;
 }
