Disabling Lumin Platform Editor modules on linux

diff --git a/Engine/Source/Developer/DesktopPlatform/Private/PlatformInfo.cpp b/Engine/Source/Developer/DesktopPlatform/Private/PlatformInfo.cpp
--- a/Engine/Source/Developer/DesktopPlatform/Private/PlatformInfo.cpp
+++ b/Engine/Source/Developer/DesktopPlatform/Private/PlatformInfo.cpp
@@ -134,7 +134,7 @@ static const FPlatformInfo AllPlatformInfoArray[] = {
 
 	BuildPlatformInfo(TEXT("TVOS"),						TEXT("TVOS"),				LOCTEXT("TVOSTargetPlatDisplay", "tvOS"),						EPlatformType::Game,	EPlatformFlags::None,			FPlatformIconPaths(TEXT("Launcher/TVOS/Platform_TVOS_24x"), TEXT("Launcher/TVOS/Platform_TVOS_128x")),								TEXT(""),											TEXT(""),			EPlatformSDKStatus::Unknown,	TEXT(""),																								IsAvailableOnWindows || IsAvailableOnMac,						TEXT("TVOS"),		TEXT("TVOS"),		false,					true,					false,			TEXT("TVOS"),		TEXT("Mobile")),
 	BuildPlatformInfo(TEXT("Switch"),					TEXT("Switch"),				LOCTEXT("Switch", "Switch"),									EPlatformType::Game,	EPlatformFlags::None,			FPlatformIconPaths(TEXT("Launcher/Switch/Platform_Switch_24x"), TEXT("Launcher/Switch/Platform_Switch_128x")),						TEXT(""),											TEXT(""),			EPlatformSDKStatus::Unknown,	TEXT(""),																								IsAvailableOnWindows,											TEXT("Switch"),		TEXT("Switch"),		false,					true,					true,			TEXT("Switch"),		TEXT("Console")),
-	BuildPlatformInfo(TEXT("Lumin"),					TEXT("Lumin"),				LOCTEXT("Lumin", "Lumin"),										EPlatformType::Game,	EPlatformFlags::None,			FPlatformIconPaths(TEXT("Launcher/LuminTarget/Platform_Lumin_24x"), TEXT("Launcher/Lumin/Platform_Lumin_128x")),					TEXT(""),											TEXT(""),			EPlatformSDKStatus::Unknown,	TEXT("/Engine/Tutorial/Mobile/SettingUpLuminTutorial.SettingUpLuminTutorial"),							IsAvailableOnLinux || IsAvailableOnWindows || IsAvailableOnMac,	TEXT("Lumin"),		TEXT("Lumin"),		false,					true,					false,			TEXT("Lumin"),		TEXT("Mobile")),
+	BuildPlatformInfo(TEXT("Lumin"),					TEXT("Lumin"),				LOCTEXT("Lumin", "Lumin"),										EPlatformType::Game,	EPlatformFlags::None,			FPlatformIconPaths(TEXT("Launcher/LuminTarget/Platform_Lumin_24x"), TEXT("Launcher/Lumin/Platform_Lumin_128x")),					TEXT(""),											TEXT(""),			EPlatformSDKStatus::Unknown,	TEXT("/Engine/Tutorial/Mobile/SettingUpLuminTutorial.SettingUpLuminTutorial"),							IsAvailableOnWindows || IsAvailableOnMac,						TEXT("Lumin"),		TEXT("Lumin"),		false,					true,					false,			TEXT("Lumin"),		TEXT("Mobile")),
 	BuildPlatformInfo(TEXT("Quail"),					TEXT("Quail"),				LOCTEXT("Quail", "Quail"),										EPlatformType::Game,	EPlatformFlags::None,			FPlatformIconPaths(TEXT("Launcher/Quail/Platform_Quail_24x"), TEXT("Launcher/Quail/Platform_Quail_128x")),							TEXT(""),											TEXT(""),			EPlatformSDKStatus::Unknown,	TEXT(""),																								IsAvailableOnWindows,											TEXT("Quail"),		TEXT("Quail"),		false,					true,					true,			TEXT("Quail"),		TEXT("")),
 
 
diff --git a/Engine/Source/Editor/UnrealEd/UnrealEd.Build.cs b/Engine/Source/Editor/UnrealEd/UnrealEd.Build.cs
--- a/Engine/Source/Editor/UnrealEd/UnrealEd.Build.cs
+++ b/Engine/Source/Editor/UnrealEd/UnrealEd.Build.cs
@@ -235,10 +235,14 @@ public class UnrealEd : ModuleRules
 		if (Target.Platform == UnrealTargetPlatform.Win64 || Target.Platform == UnrealTargetPlatform.Win32 || Target.Platform == UnrealTargetPlatform.Mac || Target.Platform == UnrealTargetPlatform.Linux)
 		{
 			DynamicallyLoadedModuleNames.Add("AndroidPlatformEditor");
+		}
+
+		if (Target.Platform == UnrealTargetPlatform.Win64 || Target.Platform == UnrealTargetPlatform.Win32 || Target.Platform == UnrealTargetPlatform.Mac)
+		{
 			DynamicallyLoadedModuleNames.Add("LuminPlatformEditor");
 		}
-        
-        CircularlyReferencedDependentModules.AddRange(
+
+		CircularlyReferencedDependentModules.AddRange(
 			new string[]
 			{
 				"GraphEditor",
