CIS - remove HTML5 AutoSDK from FPlatformInfo table

#rnx

diff --git a/Engine/Source/Developer/DesktopPlatform/Private/PlatformInfo.cpp b/Engine/Source/Developer/DesktopPlatform/Private/PlatformInfo.cpp
--- a/Engine/Source/Developer/DesktopPlatform/Private/PlatformInfo.cpp
+++ b/Engine/Source/Developer/DesktopPlatform/Private/PlatformInfo.cpp
@@ -127,7 +127,7 @@ static const FPlatformInfo AllPlatformInfoArray[] = {
 	BuildPlatformInfo(TEXT("Android_PVRTCClient"),		TEXT("Android_PVRTCClient"),LOCTEXT("Android_PVRTCClient", "Android (PVRTC) (Client-only)"),EPlatformType::Client,	EPlatformFlags::CookFlavor,		FPlatformIconPaths(TEXT("Launcher/Android/Platform_Android_PVRTC_24x"), TEXT("Launcher/Android/Platform_Android_128x")),			TEXT("-client -targetplatform=Android -cookflavor=PVRTC"),	TEXT("Android"),	EPlatformSDKStatus::Unknown,	TEXT("/Engine/Tutorial/Mobile/SettingUpAndroidTutorial.SettingUpAndroidTutorial"),				IsAvailableOnLinux || IsAvailableOnWindows || IsAvailableOnMac,	TEXT("Android"),	TEXT("Android"),	false,					true,					false,			TEXT("Android"),	TEXT("Mobile")),
 	BuildPlatformInfo(TEXT("Android_ASTCClient"),		TEXT("Android_ASTCClient"),	LOCTEXT("Android_ASTCClient", "Android (ASTC) (Client-only)"),	EPlatformType::Client,	EPlatformFlags::CookFlavor,		FPlatformIconPaths(TEXT("Launcher/Android/Platform_Android_ASTC_24x"), TEXT("Launcher/Android/Platform_Android_128x")),				TEXT("-client -targetplatform=Android -cookflavor=ASTC"),	TEXT("Android"),	EPlatformSDKStatus::Unknown,	TEXT("/Engine/Tutorial/Mobile/SettingUpAndroidTutorial.SettingUpAndroidTutorial"),				IsAvailableOnLinux || IsAvailableOnWindows || IsAvailableOnMac,	TEXT("Android"),	TEXT("Android"),	false,					true,					false,			TEXT("Android"),	TEXT("Mobile")),
 
-	BuildPlatformInfo(TEXT("HTML5"),					TEXT("HTML5"),				LOCTEXT("HTML5", "HTML5"),										EPlatformType::Game,	EPlatformFlags::None,			FPlatformIconPaths(TEXT("Launcher/HTML5/Platform_HTML5_24x"), TEXT("Launcher/HTML5/Platform_HTML5_128x")),							TEXT(""),											TEXT("HTML5"),		EPlatformSDKStatus::Unknown,	TEXT("/Platforms/HTML5/GettingStarted"),																IsAvailableOnLinux || IsAvailableOnWindows || IsAvailableOnMac,	TEXT("HTML5"),		TEXT("HTML5"),		false,					true,					false,			TEXT("HTML5"),		TEXT("Mobile")),
+	BuildPlatformInfo(TEXT("HTML5"),					TEXT("HTML5"),				LOCTEXT("HTML5", "HTML5"),										EPlatformType::Game,	EPlatformFlags::None,			FPlatformIconPaths(TEXT("Launcher/HTML5/Platform_HTML5_24x"), TEXT("Launcher/HTML5/Platform_HTML5_128x")),							TEXT(""),											TEXT(""),			EPlatformSDKStatus::Unknown,	TEXT("/Platforms/HTML5/GettingStarted"),																IsAvailableOnLinux || IsAvailableOnWindows || IsAvailableOnMac,	TEXT("HTML5"),		TEXT("HTML5"),		false,					true,					false,			TEXT("HTML5"),		TEXT("Mobile")),
 
 	BuildPlatformInfo(TEXT("PS4"),						TEXT("PS4"),				LOCTEXT("PS4", "PlayStation 4"),								EPlatformType::Game,	EPlatformFlags::None,			FPlatformIconPaths(TEXT("Launcher/PS4/Platform_PS4_24x"), TEXT("Launcher/PS4/Platform_PS4_128x")),									TEXT(""),											TEXT("PS4"),		EPlatformSDKStatus::Unknown,	TEXT("/Platforms/PS4/GettingStarted"),																	IsAvailableOnWindows,											TEXT("PS4"),		TEXT("PS4"),		false,					false,					true,			TEXT("PS4"),		TEXT("Console")),
 
