Updated FDesktopPlatformMac::FileDialogShared() to handle a case where no extensions were specified in FileTypes string

diff --git a/Engine/Source/Developer/DesktopPlatform/Private/Mac/DesktopPlatformMac.cpp b/Engine/Source/Developer/DesktopPlatform/Private/Mac/DesktopPlatformMac.cpp
--- a/Engine/Source/Developer/DesktopPlatform/Private/Mac/DesktopPlatformMac.cpp
+++ b/Engine/Source/Developer/DesktopPlatform/Private/Mac/DesktopPlatformMac.cpp
@@ -412,6 +412,11 @@ bool FDesktopPlatformMac::FileDialogShared(bool bSave, const void* ParentWindowH
 				}
 			}
 
+			if ([AllowedFileTypes count] == 1)
+			{
+				[AllowedFileTypes addObject:@""];
+			}
+			
 			[AccessoryView AddAllowedFileTypes:AllowedFileTypes];
 
 			bool bOkPressed = false;
