- Converted FPlatformInfo table to DataDrivenPlatformInfo.ini sections

diff --git a/Engine/Config/Android/DataDrivenPlatformInfo.ini b/Engine/Config/Android/DataDrivenPlatformInfo.ini
new file mode 100644
--- /dev/null
+++ b/Engine/Config/Android/DataDrivenPlatformInfo.ini
@@ -0,0 +1,345 @@
+[DataDrivenPlatformInfo]
+bIsConfidential=false
+
+[PlatformInfo Android]
+TargetPlatformName=Android
+DisplayName=Android
+PlatformType=Game
+PlatformFlags=None
+NormalIconPath=Launcher/Android/Platform_Android_24x
+LargeIconPath=Launcher/Android/Platform_Android_128x
+XLargeIconPath=
+UATCommandLine=
+AutoSDKPath=Android
+TutorialPath=/Engine/Tutorial/Mobile/SettingUpAndroidTutorial.SettingUpAndroidTutorial
+bIsEnabled=true
+BinariesDirectoryName=Android
+IniPlatformName=Android
+bUsesHostCompiler=false
+bUATClosesAfterLaunch=true
+UBTTargetID=Android
+PlatformGroupName=Mobile
+
+[PlatformInfo Android_Multi]
+TargetPlatformName=Android_Multi
+DisplayName=Android (Multi)
+PlatformType=Game
+PlatformFlags=CookFlavor
+NormalIconPath=Launcher/Android/Platform_Android_24x
+LargeIconPath=Launcher/Android/Platform_Android_128x
+XLargeIconPath=
+UATCommandLine=-targetplatform=Android -cookflavor=Multi
+AutoSDKPath=Android
+TutorialPath=/Engine/Tutorial/Mobile/SettingUpAndroidTutorial.SettingUpAndroidTutorial
+bIsEnabled=true
+BinariesDirectoryName=Android
+IniPlatformName=Android
+bUsesHostCompiler=false
+bUATClosesAfterLaunch=true
+UBTTargetID=Android
+PlatformGroupName=Mobile
+
+[PlatformInfo Android_ATC]
+TargetPlatformName=Android_ATC
+DisplayName=Android (ATC)
+PlatformType=Game
+PlatformFlags=CookFlavor
+NormalIconPath=Launcher/Android/Platform_Android_ATC_24x
+LargeIconPath=Launcher/Android/Platform_Android_ATC_128x
+XLargeIconPath=
+UATCommandLine=-targetplatform=Android -cookflavor=ATC
+AutoSDKPath=Android
+TutorialPath=/Engine/Tutorial/Mobile/SettingUpAndroidTutorial.SettingUpAndroidTutorial
+bIsEnabled=true
+BinariesDirectoryName=Android
+IniPlatformName=Android
+bUsesHostCompiler=false
+bUATClosesAfterLaunch=true
+UBTTargetID=Android
+PlatformGroupName=Mobile
+
+[PlatformInfo Android_DXT]
+TargetPlatformName=Android_DXT
+DisplayName=Android (DXT)
+PlatformType=Game
+PlatformFlags=CookFlavor
+NormalIconPath=Launcher/Android/Platform_Android_DXT_24x
+LargeIconPath=Launcher/Android/Platform_Android_DXT_128x
+XLargeIconPath=
+UATCommandLine=-targetplatform=Android -cookflavor=DXT
+AutoSDKPath=Android
+TutorialPath=/Engine/Tutorial/Mobile/SettingUpAndroidTutorial.SettingUpAndroidTutorial
+bIsEnabled=true
+BinariesDirectoryName=Android
+IniPlatformName=Android
+bUsesHostCompiler=false
+bUATClosesAfterLaunch=true
+UBTTargetID=Android
+PlatformGroupName=Mobile
+
+[PlatformInfo Android_ETC1]
+TargetPlatformName=Android_ETC1
+DisplayName=Android (ETC1)
+PlatformType=Game
+PlatformFlags=CookFlavor
+NormalIconPath=Launcher/Android/Platform_Android_ETC1_24x
+LargeIconPath=Launcher/Android/Platform_Android_ETC1_128x
+XLargeIconPath=
+UATCommandLine=-targetplatform=Android -cookflavor=ETC1
+AutoSDKPath=Android
+TutorialPath=/Engine/Tutorial/Mobile/SettingUpAndroidTutorial.SettingUpAndroidTutorial
+bIsEnabled=true
+BinariesDirectoryName=Android
+IniPlatformName=Android
+bUsesHostCompiler=false
+bUATClosesAfterLaunch=true
+UBTTargetID=Android
+PlatformGroupName=Mobile
+
+[PlatformInfo Android_ETC1a]
+TargetPlatformName=Android_ETC1a
+DisplayName=Android (ETC1a)
+PlatformType=Game
+PlatformFlags=CookFlavor
+NormalIconPath=Launcher/Android/Platform_Android_ETC1a_24x
+LargeIconPath=Launcher/Android/Platform_Android_ETC1a_128x
+XLargeIconPath=
+UATCommandLine=-targetplatform=Android -cookflavor=ETC1a
+AutoSDKPath=Android
+TutorialPath=/Engine/Tutorial/Mobile/SettingUpAndroidTutorial.SettingUpAndroidTutorial
+bIsEnabled=true
+BinariesDirectoryName=Android
+IniPlatformName=Android
+bUsesHostCompiler=false
+bUATClosesAfterLaunch=true
+UBTTargetID=Android
+PlatformGroupName=Mobile
+
+[PlatformInfo Android_ETC2]
+TargetPlatformName=Android_ETC2
+DisplayName=Android (ETC2)
+PlatformType=Game
+PlatformFlags=CookFlavor
+NormalIconPath=Launcher/Android/Platform_Android_ETC2_24x
+LargeIconPath=Launcher/Android/Platform_Android_ETC2_128x
+XLargeIconPath=
+UATCommandLine=-targetplatform=Android -cookflavor=ETC2
+AutoSDKPath=Android
+TutorialPath=/Engine/Tutorial/Mobile/SettingUpAndroidTutorial.SettingUpAndroidTutorial
+bIsEnabled=true
+BinariesDirectoryName=Android
+IniPlatformName=Android
+bUsesHostCompiler=false
+bUATClosesAfterLaunch=true
+UBTTargetID=Android
+PlatformGroupName=Mobile
+
+[PlatformInfo Android_PVRTC]
+TargetPlatformName=Android_PVRTC
+DisplayName=Android (PVRTC)
+PlatformType=Game
+PlatformFlags=CookFlavor
+NormalIconPath=Launcher/Android/Platform_Android_PVRTC_24x
+LargeIconPath=Launcher/Android/Platform_Android_PVRTC_128x
+XLargeIconPath=
+UATCommandLine=-targetplatform=Android -cookflavor=PVRTC
+AutoSDKPath=Android
+TutorialPath=/Engine/Tutorial/Mobile/SettingUpAndroidTutorial.SettingUpAndroidTutorial
+bIsEnabled=true
+BinariesDirectoryName=Android
+IniPlatformName=Android
+bUsesHostCompiler=false
+bUATClosesAfterLaunch=true
+UBTTargetID=Android
+PlatformGroupName=Mobile
+
+[PlatformInfo Android_ASTC]
+TargetPlatformName=Android_ASTC
+DisplayName=Android (ASTC)
+PlatformType=Game
+PlatformFlags=CookFlavor
+NormalIconPath=Launcher/Android/Platform_Android_ASTC_24x
+LargeIconPath=Launcher/Android/Platform_Android_ASTC_128x
+XLargeIconPath=
+UATCommandLine=-targetplatform=Android -cookflavor=ASTC
+AutoSDKPath=Android
+TutorialPath=/Engine/Tutorial/Mobile/SettingUpAndroidTutorial.SettingUpAndroidTutorial
+bIsEnabled=true
+BinariesDirectoryName=Android
+IniPlatformName=Android
+bUsesHostCompiler=false
+bUATClosesAfterLaunch=true
+UBTTargetID=Android
+PlatformGroupName=Mobile
+
+
+[PlatformInfo AndroidClient]
+TargetPlatformName=AndroidClient
+DisplayName=Android (Client-only)
+PlatformType=Game
+PlatformFlags=None
+NormalIconPath=Launcher/Android/Platform_Android_24x
+LargeIconPath=Launcher/Android/Platform_Android_128x
+XLargeIconPath=
+UATCommandLine=-client
+AutoSDKPath=Android
+TutorialPath=/Engine/Tutorial/Mobile/SettingUpAndroidTutorial.SettingUpAndroidTutorial
+bIsEnabled=true
+BinariesDirectoryName=Android
+IniPlatformName=Android
+bUsesHostCompiler=false
+bUATClosesAfterLaunch=true
+UBTTargetID=Android
+PlatformGroupName=Mobile
+
+[PlatformInfo Android_MultiClient]
+TargetPlatformName=Android_MultiClient
+DisplayName=Android (Multi) (Client-only)
+PlatformType=Game
+PlatformFlags=CookFlavor
+NormalIconPath=Launcher/Android/Platform_Android_24x
+LargeIconPath=Launcher/Android/Platform_Android_128x
+XLargeIconPath=
+UATCommandLine=-client -targetplatform=Android -cookflavor=Multi
+AutoSDKPath=Android
+TutorialPath=/Engine/Tutorial/Mobile/SettingUpAndroidTutorial.SettingUpAndroidTutorial
+bIsEnabled=true
+BinariesDirectoryName=Android
+IniPlatformName=Android
+bUsesHostCompiler=false
+bUATClosesAfterLaunch=true
+UBTTargetID=Android
+PlatformGroupName=Mobile
+
+[PlatformInfo Android_ATCClient]
+TargetPlatformName=Android_ATCClient
+DisplayName=Android (ATC) (Client-only)
+PlatformType=Game
+PlatformFlags=CookFlavor
+NormalIconPath=Launcher/Android/Platform_Android_ATC_24x
+LargeIconPath=Launcher/Android/Platform_Android_ATC_128x
+XLargeIconPath=
+UATCommandLine=-client -targetplatform=Android -cookflavor=ATC
+AutoSDKPath=Android
+TutorialPath=/Engine/Tutorial/Mobile/SettingUpAndroidTutorial.SettingUpAndroidTutorial
+bIsEnabled=true
+BinariesDirectoryName=Android
+IniPlatformName=Android
+bUsesHostCompiler=false
+bUATClosesAfterLaunch=true
+UBTTargetID=Android
+PlatformGroupName=Mobile
+
+[PlatformInfo Android_DXTClient]
+TargetPlatformName=Android_DXTClient
+DisplayName=Android (DXT) (Client-only)
+PlatformType=Game
+PlatformFlags=CookFlavor
+NormalIconPath=Launcher/Android/Platform_Android_DXT_24x
+LargeIconPath=Launcher/Android/Platform_Android_DXT_128x
+XLargeIconPath=
+UATCommandLine=-client -targetplatform=Android -cookflavor=DXT
+AutoSDKPath=Android
+TutorialPath=/Engine/Tutorial/Mobile/SettingUpAndroidTutorial.SettingUpAndroidTutorial
+bIsEnabled=true
+BinariesDirectoryName=Android
+IniPlatformName=Android
+bUsesHostCompiler=false
+bUATClosesAfterLaunch=true
+UBTTargetID=Android
+PlatformGroupName=Mobile
+
+[PlatformInfo Android_ETC1Client]
+TargetPlatformName=Android_ETC1Client
+DisplayName=Android (ETC1) (Client-only)
+PlatformType=Game
+PlatformFlags=CookFlavor
+NormalIconPath=Launcher/Android/Platform_Android_ETC1_24x
+LargeIconPath=Launcher/Android/Platform_Android_ETC1_128x
+XLargeIconPath=
+UATCommandLine=-client -targetplatform=Android -cookflavor=ETC1
+AutoSDKPath=Android
+TutorialPath=/Engine/Tutorial/Mobile/SettingUpAndroidTutorial.SettingUpAndroidTutorial
+bIsEnabled=true
+BinariesDirectoryName=Android
+IniPlatformName=Android
+bUsesHostCompiler=false
+bUATClosesAfterLaunch=true
+UBTTargetID=Android
+PlatformGroupName=Mobile
+
+[PlatformInfo Android_ETC1aClient]
+TargetPlatformName=Android_ETC1aClient
+DisplayName=Android (ETC1a) (Client-only)
+PlatformType=Game
+PlatformFlags=CookFlavor
+NormalIconPath=Launcher/Android/Platform_Android_ETC1a_24x
+LargeIconPath=Launcher/Android/Platform_Android_ETC1a_128x
+XLargeIconPath=
+UATCommandLine=-client -targetplatform=Android -cookflavor=ETC1a
+AutoSDKPath=Android
+TutorialPath=/Engine/Tutorial/Mobile/SettingUpAndroidTutorial.SettingUpAndroidTutorial
+bIsEnabled=true
+BinariesDirectoryName=Android
+IniPlatformName=Android
+bUsesHostCompiler=false
+bUATClosesAfterLaunch=true
+UBTTargetID=Android
+PlatformGroupName=Mobile
+
+[PlatformInfo Android_ETC2Client]
+TargetPlatformName=Android_ETC2Client
+DisplayName=Android (ETC2) (Client-only)
+PlatformType=Game
+PlatformFlags=CookFlavor
+NormalIconPath=Launcher/Android/Platform_Android_ETC2_24x
+LargeIconPath=Launcher/Android/Platform_Android_ETC2_128x
+XLargeIconPath=
+UATCommandLine=-client -targetplatform=Android -cookflavor=ETC2
+AutoSDKPath=Android
+TutorialPath=/Engine/Tutorial/Mobile/SettingUpAndroidTutorial.SettingUpAndroidTutorial
+bIsEnabled=true
+BinariesDirectoryName=Android
+IniPlatformName=Android
+bUsesHostCompiler=false
+bUATClosesAfterLaunch=true
+UBTTargetID=Android
+PlatformGroupName=Mobile
+
+[PlatformInfo Android_PVRTCClient]
+TargetPlatformName=Android_PVRTCClient
+DisplayName=Android (PVRTC) (Client-only)
+PlatformType=Game
+PlatformFlags=CookFlavor
+NormalIconPath=Launcher/Android/Platform_Android_PVRTC_24x
+LargeIconPath=Launcher/Android/Platform_Android_PVRTC_128x
+XLargeIconPath=
+UATCommandLine=-client -targetplatform=Android -cookflavor=PVRTC
+AutoSDKPath=Android
+TutorialPath=/Engine/Tutorial/Mobile/SettingUpAndroidTutorial.SettingUpAndroidTutorial
+bIsEnabled=true
+BinariesDirectoryName=Android
+IniPlatformName=Android
+bUsesHostCompiler=false
+bUATClosesAfterLaunch=true
+UBTTargetID=Android
+PlatformGroupName=Mobile
+
+[PlatformInfo Android_ASTCClient]
+TargetPlatformName=Android_ASTCClient
+DisplayName=Android (ASTC) (Client-only)
+PlatformType=Game
+PlatformFlags=CookFlavor
+NormalIconPath=Launcher/Android/Platform_Android_ASTC_24x
+LargeIconPath=Launcher/Android/Platform_Android_ASTC_128x
+XLargeIconPath=
+UATCommandLine=-client -targetplatform=Android -cookflavor=ASTC
+AutoSDKPath=Android
+TutorialPath=/Engine/Tutorial/Mobile/SettingUpAndroidTutorial.SettingUpAndroidTutorial
+bIsEnabled=true
+BinariesDirectoryName=Android
+IniPlatformName=Android
+bUsesHostCompiler=false
+bUATClosesAfterLaunch=true
+UBTTargetID=Android
+PlatformGroupName=Mobile
diff --git a/Engine/Config/HTML5/DataDrivenPlatformInfo.ini b/Engine/Config/HTML5/DataDrivenPlatformInfo.ini
new file mode 100644
--- /dev/null
+++ b/Engine/Config/HTML5/DataDrivenPlatformInfo.ini
@@ -0,0 +1,40 @@
+[DataDrivenPlatformInfo]
+bIsConfidential=false
+
+[PlatformInfo HTML5]
+TargetPlatformName=HTML5
+DisplayName=HTML5
+PlatformType=Game
+PlatformFlags=None
+NormalIconPath=Launcher/HTML5/Platform_HTML5_24x
+LargeIconPath=Launcher/HTML5/Platform_HTML5_128x
+XLargeIconPath=
+UATCommandLine=
+AutoSDKPath=
+TutorialPath=/Platforms/HTML5/GettingStarted
+bIsEnabled=true
+BinariesDirectoryName=HTML5
+IniPlatformName=HTML5
+bUsesHostCompiler=false
+bUATClosesAfterLaunch=true
+UBTTargetID=HTML5
+PlatformGroupName=Mobile
+
+[PlatformInfo HTML5Client]
+TargetPlatformName=HTML5
+DisplayName=HTML5 (Client-only)
+PlatformType=Client
+PlatformFlags=None
+NormalIconPath=Launcher/HTML5/Platform_HTML5_24x
+LargeIconPath=Launcher/HTML5/Platform_HTML5_128x
+XLargeIconPath=
+UATCommandLine=-client
+AutoSDKPath=
+TutorialPath=/Platforms/HTML5/GettingStarted
+bIsEnabled=true
+BinariesDirectoryName=HTML5
+IniPlatformName=HTML5
+bUsesHostCompiler=false
+bUATClosesAfterLaunch=true
+UBTTargetID=HTML5
+PlatformGroupName=Mobile
diff --git a/Engine/Config/IOS/DataDrivenPlatformInfo.ini b/Engine/Config/IOS/DataDrivenPlatformInfo.ini
new file mode 100644
--- /dev/null
+++ b/Engine/Config/IOS/DataDrivenPlatformInfo.ini
@@ -0,0 +1,48 @@
+[DataDrivenPlatformInfo]
+bIsConfidential=false
+
+[PlatformInfo IOS]
+TargetPlatformName=IOS
+DisplayName=iOS
+PlatformType=Game
+PlatformFlags=None
+NormalIconPath=Launcher/IOS/Platform_IOS_24x
+LargeIconPath=Launcher/IOS/Platform_IOS_128x
+XLargeIconPath=
+UATCommandLine=
+AutoSDKPath=
+Windows:TutorialPath=/Engine/Tutorial/Mobile/InstallingiTunesTutorial.InstallingiTunesTutorial
+Mac:TutorialPath=/Engine/Tutorial/Installation/InstallingXCodeTutorial.InstallingXCodeTutorial
+Linux:TutorialPath=/Engine/Tutorial/NotYetImplemented
+Windows:bIsEnabled=true
+Mac:bIsEnabled=true
+Linux:bIsEnabled=false
+BinariesDirectoryName=IOS
+IniPlatformName=IOS
+bUsesHostCompiler=false
+bUATClosesAfterLaunch=true
+UBTTargetID=IOS
+PlatformGroupName=Mobile
+
+[PlatformInfo IOSClient]
+TargetPlatformName=IOSClient
+DisplayName=iOS (Client-only)
+PlatformType=Client
+PlatformFlags=None
+NormalIconPath=Launcher/IOS/Platform_IOS_24x
+LargeIconPath=Launcher/IOS/Platform_IOS_24x
+XLargeIconPath=
+UATCommandLine=-client
+AutoSDKPath=
+Windows:TutorialPath=/Engine/Tutorial/Mobile/InstallingiTunesTutorial.InstallingiTunesTutorial
+Mac:TutorialPath=/Engine/Tutorial/Installation/InstallingXCodeTutorial.InstallingXCodeTutorial
+Linux:TutorialPath=/Engine/Tutorial/NotYetImplemented
+Windows:bIsEnabled=true
+Mac:bIsEnabled=true
+Linux:bIsEnabled=false
+BinariesDirectoryName=IOS
+IniPlatformName=IOS
+bUsesHostCompiler=false
+bUATClosesAfterLaunch=true
+UBTTargetID=IOS
+PlatformGroupName=Mobile
diff --git a/Engine/Config/Linux/DataDrivenPlatformInfo.ini b/Engine/Config/Linux/DataDrivenPlatformInfo.ini
--- a/Engine/Config/Linux/DataDrivenPlatformInfo.ini
+++ b/Engine/Config/Linux/DataDrivenPlatformInfo.ini
@@ -1,2 +1,91 @@
 [DataDrivenPlatformInfo]
 IniParent=Unix
+
+[PlatformInfo Linux]
+TargetPlatformName=Linux
+DisplayName=Linux (Editor)
+PlatformType=Editor
+PlatformFlags=None
+NormalIconPath=Launcher/Linux/Platform_Linux_24x
+LargeIconPath=Launcher/Linux/Platform_Linux_128x
+XLargeIconPath=
+UATCommandLine=
+AutoSDKPath=Linux_x64
+TutorialPath=/Platforms/Linux/GettingStarted
+Windows:bIsEnabled=false
+Mac:bIsEnabled=false
+Linux:bIsEnabled=true
+BinariesDirectoryName=Linux
+IniPlatformName=Linux
+Windows:bUsesHostCompiler=false
+Linux:bUsesHostCompiler=true
+bUATClosesAfterLaunch=true
+UBTTargetID=Linux
+PlatformGroupName=Desktop
+
+[PlatformInfo LinuxNoEditor]
+TargetPlatformName=LinuxNoEditor
+DisplayName=Linux
+PlatformType=Game
+PlatformFlags=None
+NormalIconPath=Launcher/Linux/Platform_Linux_24x
+LargeIconPath=Launcher/Linux/Platform_Linux_128x
+XLargeIconPath=
+UATCommandLine=-targetplatform=Linux
+AutoSDKPath=Linux_x64
+TutorialPath=/Platforms/Linux/GettingStarted
+Windows:bIsEnabled=true
+Mac:bIsEnabled=false
+Linux:bIsEnabled=true
+BinariesDirectoryName=Linux
+IniPlatformName=Linux
+Windows:bUsesHostCompiler=false
+Linux:bUsesHostCompiler=true
+bUATClosesAfterLaunch=true
+UBTTargetID=Linux
+PlatformGroupName=Desktop
+
+[PlatformInfo LinuxClient]
+TargetPlatformName=LinuxClient
+DisplayName=Linux (Client-only)
+PlatformType=Client
+PlatformFlags=None
+NormalIconPath=Launcher/Linux/Platform_Linux_24x
+LargeIconPath=Launcher/Linux/Platform_Linux_128x
+XLargeIconPath=
+UATCommandLine=-client
+AutoSDKPath=Linux_x64
+TutorialPath=/Platforms/Linux/GettingStarted
+Windows:bIsEnabled=true
+Mac:bIsEnabled=false
+Linux:bIsEnabled=true
+BinariesDirectoryName=Linux
+IniPlatformName=Linux
+Windows:bUsesHostCompiler=false
+Linux:bUsesHostCompiler=true
+bUATClosesAfterLaunch=true
+UBTTargetID=Linux
+PlatformGroupName=Desktop
+
+[PlatformInfo LinuxServer]
+TargetPlatformName=LinuxServer
+DisplayName=Linux (Dedicated Server)
+PlatformType=Editor
+PlatformFlags=None
+NormalIconPath=Launcher/Linux/Platform_Linux_24x
+LargeIconPath=Launcher/Linux/Platform_Linux_128x
+XLargeIconPath=
+UATCommandLine=
+AutoSDKPath=Linux_x64
+TutorialPath=/Platforms/Linux/GettingStarted
+Windows:bIsEnabled=true
+Mac:bIsEnabled=false
+Linux:bIsEnabled=true
+BinariesDirectoryName=Linux
+IniPlatformName=Linux
+Windows:bUsesHostCompiler=false
+Linux:bUsesHostCompiler=true
+bUATClosesAfterLaunch=true
+UBTTargetID=Linux
+PlatformGroupName=Desktop
+
diff --git a/Engine/Config/Lumin/DataDrivenPlatformInfo.ini b/Engine/Config/Lumin/DataDrivenPlatformInfo.ini
--- a/Engine/Config/Lumin/DataDrivenPlatformInfo.ini
+++ b/Engine/Config/Lumin/DataDrivenPlatformInfo.ini
@@ -1,3 +1,46 @@
 [DataDrivenPlatformInfo]
 bIsConfidential=false
 IniParent=Android
+
+
+[PlatformInfo Lumin]
+TargetPlatformName=Lumin
+DisplayName=Lumin
+PlatformType=Game
+PlatformFlags=None
+NormalIconPath=Launcher/Lumin/Platform_Lumin_24x
+LargeIconPath=Launcher/Lumin/Platform_Lumin_128x
+XLargeIconPath=
+UATCommandLine=
+AutoSDKPath=
+TutorialPath=/Engine/Tutorial/Mobile/SettingUpLuminTutorial.SettingUpLuminTutorial
+Windows:bIsEnabled=true
+Mac:bIsEnabled=true
+Linux:bIsEnabled=false
+BinariesDirectoryName=Lumin
+IniPlatformName=Lumin
+bUsesHostCompiler=false
+bUATClosesAfterLaunch=true
+UBTTargetID=Lumin
+PlatformGroupName=Mobile
+
+[PlatformInfo LuminClient]
+TargetPlatformName=LuminClient
+DisplayName=Lumin (Client-only)
+PlatformType=Client
+PlatformFlags=None
+NormalIconPath=Launcher/Lumin/Platform_Lumin_24x
+LargeIconPath=Launcher/Lumin/Platform_Lumin_128x
+XLargeIconPath=
+UATCommandLine=-client
+AutoSDKPath=
+TutorialPath=/Engine/Tutorial/Mobile/SettingUpLuminTutorial.SettingUpLuminTutorial
+Windows:bIsEnabled=true
+Mac:bIsEnabled=true
+Linux:bIsEnabled=false
+BinariesDirectoryName=Lumin
+IniPlatformName=Lumin
+bUsesHostCompiler=false
+bUATClosesAfterLaunch=true
+UBTTargetID=Lumin
+PlatformGroupName=Mobile
diff --git a/Engine/Config/Mac/DataDrivenPlatformInfo.ini b/Engine/Config/Mac/DataDrivenPlatformInfo.ini
new file mode 100644
--- /dev/null
+++ b/Engine/Config/Mac/DataDrivenPlatformInfo.ini
@@ -0,0 +1,86 @@
+[DataDrivenPlatformInfo]
+bIsConfidential=false
+
+[PlatformInfo MacNoEditor]
+TargetPlatformName=MacNoEditor
+DisplayName=Mac
+PlatformType=Game
+PlatformFlags=None
+NormalIconPath=Launcher/Mac/Platform_Mac_24x
+LargeIconPath=Launcher/Mac/Platform_Mac_128x
+XLargeIconPath=
+UATCommandLine=-targetplatform=Mac
+AutoSDKPath=
+TutorialPath=/Engine/Tutorial/Installation/InstallingXCodeTutorial.InstallingXCodeTutorial
+Windows:bIsEnabled=false
+Mac:bIsEnabled=true
+Linux:bIsEnabled=false
+BinariesDirectoryName=Mac
+IniPlatformName=Mac
+bUsesHostCompiler=true
+bUATClosesAfterLaunch=true
+UBTTargetID=Mac
+PlatformGroupName=Desktop
+
+[PlatformInfo Mac]
+TargetPlatformName=Mac
+DisplayName=Mac (Editor)
+PlatformType=Editor
+PlatformFlags=None
+NormalIconPath=Launcher/Mac/Platform_Mac_24x
+LargeIconPath=Launcher/Mac/Platform_Mac_128x
+XLargeIconPath=
+UATCommandLine=
+AutoSDKPath=
+TutorialPath=/Engine/Tutorial/Installation/InstallingXCodeTutorial.InstallingXCodeTutorial
+Windows:bIsEnabled=false
+Mac:bIsEnabled=true
+Linux:bIsEnabled=false
+BinariesDirectoryName=Mac
+IniPlatformName=Mac
+bUsesHostCompiler=true
+bUATClosesAfterLaunch=true
+UBTTargetID=Mac
+PlatformGroupName=Desktop
+
+[PlatformInfo MacClient]
+TargetPlatformName=MacClient
+DisplayName=Mac (Client-only)
+PlatformType=Client
+PlatformFlags=None
+NormalIconPath=Launcher/Mac/Platform_Mac_24x
+LargeIconPath=Launcher/Mac/Platform_Mac_128x
+XLargeIconPath=
+UATCommandLine=-client
+AutoSDKPath=
+TutorialPath=/Engine/Tutorial/Installation/InstallingXCodeTutorial.InstallingXCodeTutorial
+Windows:bIsEnabled=false
+Mac:bIsEnabled=true
+Linux:bIsEnabled=false
+BinariesDirectoryName=Mac
+IniPlatformName=Mac
+bUsesHostCompiler=true
+bUATClosesAfterLaunch=true
+UBTTargetID=Mac
+PlatformGroupName=Desktop
+
+[PlatformInfo MacServer]
+TargetPlatformName=MacServer
+DisplayName=Mac (Dedicated Server)
+PlatformType=Server
+PlatformFlags=None
+NormalIconPath=Launcher/Mac/Platform_Mac_24x
+LargeIconPath=Launcher/Mac/Platform_Mac_128x
+XLargeIconPath=
+UATCommandLine=
+AutoSDKPath=
+TutorialPath=/Engine/Tutorial/Installation/InstallingXCodeTutorial.InstallingXCodeTutorial
+Windows:bIsEnabled=false
+Mac:bIsEnabled=true
+Linux:bIsEnabled=false
+BinariesDirectoryName=Mac
+IniPlatformName=Mac
+bUsesHostCompiler=true
+bUATClosesAfterLaunch=true
+UBTTargetID=Mac
+PlatformGroupName=Desktop
diff --git a/Engine/Config/TVOS/DataDrivenPlatformInfo.ini b/Engine/Config/TVOS/DataDrivenPlatformInfo.ini
new file mode 100644
--- /dev/null
+++ b/Engine/Config/TVOS/DataDrivenPlatformInfo.ini
@@ -0,0 +1,45 @@
+[DataDrivenPlatformInfo]
+bIsConfidential=false
+IniParent=IOS
+
+[PlatformInfo TVOS]
+TargetPlatformName=TVOS
+DisplayName=tvOS
+PlatformType=Game
+PlatformFlags=None
+NormalIconPath=Launcher/TVOS/Platform_TVOS_24x
+LargeIconPath=Launcher/TVOS/Platform_TVOS_128x
+XLargeIconPath=
+UATCommandLine=
+AutoSDKPath=
+TutorialPath=
+Windows:bIsEnabled=true
+Mac:bIsEnabled=true
+Linux:bIsEnabled=false
+BinariesDirectoryName=TVOS
+IniPlatformName=TVOS
+bUsesHostCompiler=false
+bUATClosesAfterLaunch=true
+UBTTargetID=TVOS
+PlatformGroupName=Mobile
+
+[PlatformInfo TVOSClient]
+TargetPlatformName=TVOSClient
+DisplayName=tvOS (Client-only)
+PlatformType=Client
+PlatformFlags=None
+NormalIconPath=Launcher/TVOS/Platform_TVOS_24x
+LargeIconPath=Launcher/TVOS/Platform_TVOS_24x
+XLargeIconPath=
+UATCommandLine=-client
+AutoSDKPath=
+TutorialPath=
+Windows:bIsEnabled=true
+Mac:bIsEnabled=true
+Linux:bIsEnabled=false
+BinariesDirectoryName=TVOS
+IniPlatformName=TVOS
+bUsesHostCompiler=false
+bUATClosesAfterLaunch=true
+UBTTargetID=TVOS
+PlatformGroupName=Mobile
diff --git a/Engine/Config/Windows/DataDrivenPlatformInfo.ini b/Engine/Config/Windows/DataDrivenPlatformInfo.ini
new file mode 100644
--- /dev/null
+++ b/Engine/Config/Windows/DataDrivenPlatformInfo.ini
@@ -0,0 +1,128 @@
+[DataDrivenPlatformInfo]
+bIsConfidential=false
+
+[PlatformInfo WindowsNoEditor]
+TargetPlatformName=WindowsNoEditor
+DisplayName=Windows
+PlatformType=Game
+PlatformFlags=None
+NormalIconPath=Launcher/Windows/Platform_WindowsNoEditor_24x
+LargeIconPath=Launcher/Windows/Platform_WindowsNoEditor_128x
+XLargeIconPath=
+UATCommandLine=-targetplatform=Win64
+AutoSDKPath=
+TutorialPath=/Engine/Tutorial/Installation/InstallingVisualStudioTutorial.InstallingVisualStudioTutorial
+Windows:bIsEnabled=true
+Mac:bIsEnabled=false
+Linux:bIsEnabled=false
+BinariesDirectoryName=Win64
+IniPlatformName=Windows
+bUsesHostCompiler=true
+bUATClosesAfterLaunch=false
+UBTTargetID=Win64
+PlatformGroupName=Desktop
+
+[PlatformInfo WindowsNoEditor_Win32]
+TargetPlatformName=WindowsNoEditor
+DisplayName=Windows (32-bit)
+PlatformType=Game
+PlatformFlags=BuildFlavor
+NormalIconPath=Launcher/Windows/Platform_WindowsNoEditor_24x
+LargeIconPath=Launcher/Windows/Platform_WindowsNoEditor_128x
+XLargeIconPath=
+UATCommandLine=-targetplatform=Win32
+AutoSDKPath=
+TutorialPath=/Engine/Tutorial/Installation/InstallingVisualStudioTutorial.InstallingVisualStudioTutorial
+Windows:bIsEnabled=true
+Mac:bIsEnabled=false
+Linux:bIsEnabled=false
+BinariesDirectoryName=Win32
+IniPlatformName=Windows
+bUsesHostCompiler=true
+bUATClosesAfterLaunch=false
+UBTTargetID=Win32
+PlatformGroupName=Desktop
+
+[PlatformInfo WindowsNoEditor_Win64]
+TargetPlatformName=WindowsNoEditor
+DisplayName=Windows (64-bit)
+PlatformType=Game
+PlatformFlags=BuildFlavor
+NormalIconPath=Launcher/Windows/Platform_WindowsNoEditor_24x
+LargeIconPath=Launcher/Windows/Platform_WindowsNoEditor_128x
+XLargeIconPath=
+UATCommandLine=-targetplatform=Win64
+AutoSDKPath=
+TutorialPath=/Engine/Tutorial/Installation/InstallingVisualStudioTutorial.InstallingVisualStudioTutorial
+Windows:bIsEnabled=true
+Mac:bIsEnabled=false
+Linux:bIsEnabled=false
+BinariesDirectoryName=Win64
+IniPlatformName=Windows
+bUsesHostCompiler=true
+bUATClosesAfterLaunch=false
+UBTTargetID=Win64
+PlatformGroupName=Desktop
+
+[PlatformInfo Windows]
+TargetPlatformName=Windows
+DisplayName=Windows (Editor)
+PlatformType=Editor
+PlatformFlags=None
+NormalIconPath=Launcher/Windows/Platform_Windows_24x
+LargeIconPath=Launcher/Windows/Platform_Windows_128x
+XLargeIconPath=
+UATCommandLine=
+AutoSDKPath=
+TutorialPath=/Engine/Tutorial/Installation/InstallingVisualStudioTutorial.InstallingVisualStudioTutorial
+Windows:bIsEnabled=true
+Mac:bIsEnabled=false
+Linux:bIsEnabled=false
+BinariesDirectoryName=Win64
+IniPlatformName=Windows
+bUsesHostCompiler=true
+bUATClosesAfterLaunch=false
+UBTTargetID=Win64
+PlatformGroupName=Desktop
+
+[PlatformInfo WindowsClient]
+TargetPlatformName=WindowsClient
+DisplayName=Windows (Client-only)
+PlatformType=Client
+PlatformFlags=None
+NormalIconPath=Launcher/Windows/Platform_Windows_24x
+LargeIconPath=Launcher/Windows/Platform_Windows_128x
+XLargeIconPath=
+UATCommandLine=-client
+AutoSDKPath=
+TutorialPath=/Engine/Tutorial/Installation/InstallingVisualStudioTutorial.InstallingVisualStudioTutorial
+Windows:bIsEnabled=true
+Mac:bIsEnabled=false
+Linux:bIsEnabled=false
+BinariesDirectoryName=Win64
+IniPlatformName=Windows
+bUsesHostCompiler=true
+bUATClosesAfterLaunch=false
+UBTTargetID=Win64
+PlatformGroupName=Desktop
+
+[PlatformInfo WindowsServer]
+TargetPlatformName=WindowsServer
+DisplayName=Windows (Dedicated Server)
+PlatformType=Server
+PlatformFlags=None
+NormalIconPath=Launcher/Windows/Platform_WindowsServer_24x
+LargeIconPath=Launcher/Windows/Platform_WindowsServer_128x
+XLargeIconPath=
+UATCommandLine=
+AutoSDKPath=
+TutorialPath=/Engine/Tutorial/Installation/InstallingVisualStudioTutorial.InstallingVisualStudioTutorial
+Windows:bIsEnabled=true
+Mac:bIsEnabled=false
+Linux:bIsEnabled=false
+BinariesDirectoryName=Win64
+IniPlatformName=Windows
+bUsesHostCompiler=true
+bUATClosesAfterLaunch=false
+UBTTargetID=Win64
+PlatformGroupName=Desktop
diff --git a/Engine/Source/Developer/DesktopPlatform/Private/PlatformInfo.cpp b/Engine/Source/Developer/DesktopPlatform/Private/PlatformInfo.cpp
--- a/Engine/Source/Developer/DesktopPlatform/Private/PlatformInfo.cpp
+++ b/Engine/Source/Developer/DesktopPlatform/Private/PlatformInfo.cpp
@@ -77,95 +77,21 @@ void BuildPlatformInfo(const FName& InPlatformInfoName, const FName& InTargetPla
 	PlatformInfo.bIsConfidential = InIsConfidential;
 }
 
-#if PLATFORM_WINDOWS
-static const bool IsAvailableOnWindows = true;
-static const bool IsAvailableOnMac = false;
-static const bool IsAvailableOnLinux = false;
-static const FString IOSTutorial = TEXT("/Engine/Tutorial/Mobile/InstallingiTunesTutorial.InstallingiTunesTutorial");
-#elif PLATFORM_MAC
-static const bool IsAvailableOnWindows = false;
-static const bool IsAvailableOnMac = true;
-static const bool IsAvailableOnLinux = false;
-static const FString IOSTutorial = TEXT("/Engine/Tutorial/Installation/InstallingXCodeTutorial.InstallingXCodeTutorial");
-#elif PLATFORM_LINUX
-static const bool IsAvailableOnWindows = false;
-static const bool IsAvailableOnMac = false;
-static const bool IsAvailableOnLinux = true;
-static const FString IOSTutorial = TEXT("/Engine/Tutorial/NotYetImplemented");
-#else
-static const bool IsAvailableOnWindows = false;
-static const bool IsAvailableOnMac = false;
-static const bool IsAvailableOnLinux = false;
-static const FString IOSTutorial = TEXT("/Engine/Tutorial/NotYetImplemented");
-#endif
 
 void BuildHardcodedPlatforms()
 {
 	// PlatformInfoName									TargetPlatformName			DisplayName														PlatformType			PlatformFlags					IconPaths																																		UATCommandLine										AutoSDKPath			SDKStatus						SDKTutorial																								bEnabledForUse										BinaryFolderName	IniPlatformName		FbUsesHostCompiler		bUATClosesAfterLaunch	bIsConfidential	UBTTargetId (match UBT's UnrealTargetPlatform enum)
-	BuildPlatformInfo(TEXT("WindowsNoEditor"),			TEXT("WindowsNoEditor"),	LOCTEXT("WindowsNoEditor", "Windows"),							EPlatformType::Game,	EPlatformFlags::None,			FPlatformIconPaths(TEXT("Launcher/Windows/Platform_WindowsNoEditor_24x"), TEXT("Launcher/Windows/Platform_WindowsNoEditor_128x")),	TEXT("-targetplatform=Win64"),						TEXT(""),			EPlatformSDKStatus::Unknown,	TEXT("/Engine/Tutorial/Installation/InstallingVisualStudioTutorial.InstallingVisualStudioTutorial"),	IsAvailableOnWindows,											TEXT("Win64"),		TEXT("Windows"),	IsAvailableOnWindows,	false,					false,			TEXT("Win64"),		TEXT("Desktop"));
-	BuildPlatformInfo(TEXT("WindowsNoEditor_Win32"),	TEXT("WindowsNoEditor"),	LOCTEXT("WindowsNoEditor_Win32", "Windows (32-bit)"),			EPlatformType::Game,	EPlatformFlags::BuildFlavor,	FPlatformIconPaths(TEXT("Launcher/Windows/Platform_WindowsNoEditor_24x"), TEXT("Launcher/Windows/Platform_WindowsNoEditor_128x")),	TEXT("-targetplatform=Win32"),						TEXT(""),			EPlatformSDKStatus::Unknown,	TEXT("/Engine/Tutorial/Installation/InstallingVisualStudioTutorial.InstallingVisualStudioTutorial"),	IsAvailableOnWindows,											TEXT("Win32"),		TEXT("Windows"),	IsAvailableOnWindows,	false,					false,			TEXT("Win32"),		TEXT("Desktop"));
-	BuildPlatformInfo(TEXT("WindowsNoEditor_Win64"),	TEXT("WindowsNoEditor"),	LOCTEXT("WindowsNoEditor_Win64", "Windows (64-bit)"),			EPlatformType::Game,	EPlatformFlags::BuildFlavor,	FPlatformIconPaths(TEXT("Launcher/Windows/Platform_WindowsNoEditor_24x"), TEXT("Launcher/Windows/Platform_WindowsNoEditor_128x")),	TEXT("-targetplatform=Win64"),						TEXT(""),			EPlatformSDKStatus::Unknown,	TEXT("/Engine/Tutorial/Installation/InstallingVisualStudioTutorial.InstallingVisualStudioTutorial"),	IsAvailableOnWindows,											TEXT("Win64"),		TEXT("Windows"),	IsAvailableOnWindows,	false,					false,			TEXT("Win64"),		TEXT("Desktop"));
-	BuildPlatformInfo(TEXT("Windows"),					TEXT("Windows"),			LOCTEXT("WindowsEditor", "Windows (Editor)"),					EPlatformType::Editor,	EPlatformFlags::None,			FPlatformIconPaths(TEXT("Launcher/Windows/Platform_Windows_24x"), TEXT("Launcher/Windows/Platform_Windows_128x")),					TEXT(""),											TEXT(""),			EPlatformSDKStatus::Unknown,	TEXT("/Engine/Tutorial/Installation/InstallingVisualStudioTutorial.InstallingVisualStudioTutorial"),	IsAvailableOnWindows,											TEXT("Win64"),		TEXT("Windows"),	IsAvailableOnWindows,	false,					false,			TEXT("Win64"),		TEXT("Desktop"));
-	BuildPlatformInfo(TEXT("WindowsClient"),			TEXT("WindowsClient"),		LOCTEXT("WindowsClient", "Windows (Client-only)"),				EPlatformType::Client,	EPlatformFlags::None,			FPlatformIconPaths(TEXT("Launcher/Windows/Platform_Windows_24x"), TEXT("Launcher/Windows/Platform_Windows_128x")),					TEXT("-client"),									TEXT(""),			EPlatformSDKStatus::Unknown,	TEXT("/Engine/Tutorial/Installation/InstallingVisualStudioTutorial.InstallingVisualStudioTutorial"),	IsAvailableOnWindows,											TEXT("Win64"),		TEXT("Windows"),	IsAvailableOnWindows,	false,					false,			TEXT("Win64"),		TEXT("Desktop"));
-	BuildPlatformInfo(TEXT("WindowsServer"),			TEXT("WindowsServer"),		LOCTEXT("WindowsServer", "Windows (Dedicated Server)"),			EPlatformType::Server,	EPlatformFlags::None,			FPlatformIconPaths(TEXT("Launcher/Windows/Platform_WindowsServer_24x"), TEXT("Launcher/Windows/Platform_WindowsServer_128x")),		TEXT(""),											TEXT(""),			EPlatformSDKStatus::Unknown,	TEXT("/Engine/Tutorial/Installation/InstallingVisualStudioTutorial.InstallingVisualStudioTutorial"),	IsAvailableOnWindows,											TEXT("Win64"),		TEXT("Windows"),	IsAvailableOnWindows,	false,					false,			TEXT("Win64"),		TEXT("Desktop"));
-
-	BuildPlatformInfo(TEXT("MacNoEditor"),				TEXT("MacNoEditor"),		LOCTEXT("MacNoEditor", "Mac"),									EPlatformType::Game,	EPlatformFlags::None,			FPlatformIconPaths(TEXT("Launcher/Mac/Platform_Mac_24x"), TEXT("Launcher/Mac/Platform_Mac_128x")),									TEXT("-targetplatform=Mac"),						TEXT(""),			EPlatformSDKStatus::Unknown,	TEXT("/Engine/Tutorial/Installation/InstallingXCodeTutorial.InstallingXCodeTutorial"),					IsAvailableOnMac,												TEXT("Mac"),		TEXT("Mac"),		IsAvailableOnMac,		true,					false,			TEXT("Mac"),		TEXT("Desktop"));
-	BuildPlatformInfo(TEXT("Mac"),						TEXT("Mac"),				LOCTEXT("MacEditor", "Mac (Editor)"),							EPlatformType::Editor,	EPlatformFlags::None,			FPlatformIconPaths(TEXT("Launcher/Mac/Platform_Mac_24x"), TEXT("Launcher/Mac/Platform_Mac_128x")),									TEXT(""),											TEXT(""),			EPlatformSDKStatus::Unknown,	TEXT("/Engine/Tutorial/Installation/InstallingXCodeTutorial.InstallingXCodeTutorial"),					IsAvailableOnMac,												TEXT("Mac"),		TEXT("Mac"),		IsAvailableOnMac,		true,					false,			TEXT("Mac"),		TEXT("Desktop"));
-	BuildPlatformInfo(TEXT("MacClient"),				TEXT("MacClient"),			LOCTEXT("MacClient", "Mac (Client-only)"),						EPlatformType::Client,	EPlatformFlags::None,			FPlatformIconPaths(TEXT("Launcher/Mac/Platform_Mac_24x"), TEXT("Launcher/Mac/Platform_Mac_128x")),									TEXT("-client"),									TEXT(""),			EPlatformSDKStatus::Unknown,	TEXT("/Engine/Tutorial/Installation/InstallingXCodeTutorial.InstallingXCodeTutorial"),					IsAvailableOnMac,												TEXT("Mac"),		TEXT("Mac"),		IsAvailableOnMac,		true,					false,			TEXT("Mac"),		TEXT("Desktop"));
-	BuildPlatformInfo(TEXT("MacServer"),				TEXT("MacServer"),			LOCTEXT("MacServer", "Mac (Dedicated Server)"),					EPlatformType::Server,	EPlatformFlags::None,			FPlatformIconPaths(TEXT("Launcher/Mac/Platform_Mac_24x"), TEXT("Launcher/Mac/Platform_Mac_128x")),									TEXT(""),											TEXT(""),			EPlatformSDKStatus::Unknown,	TEXT("/Engine/Tutorial/Installation/InstallingXCodeTutorial.InstallingXCodeTutorial"),					IsAvailableOnMac,												TEXT("Mac"),		TEXT("Mac"),		IsAvailableOnMac,		true,					false,			TEXT("Mac"),		TEXT("Desktop"));
-
-	BuildPlatformInfo(TEXT("LinuxNoEditor"),			TEXT("LinuxNoEditor"),		LOCTEXT("LinuxNoEditor", "Linux"),								EPlatformType::Game,	EPlatformFlags::None,			FPlatformIconPaths(TEXT("Launcher/Linux/Platform_Linux_24x"), TEXT("Launcher/Linux/Platform_Linux_128x")),							TEXT("-targetplatform=Linux"),						TEXT("Linux_x64"),	EPlatformSDKStatus::Unknown,	TEXT("/Platforms/Linux/GettingStarted"),																IsAvailableOnLinux || IsAvailableOnWindows,						TEXT("Linux"),		TEXT("Linux"),		IsAvailableOnLinux,		true,					false,			TEXT("Linux"),		TEXT("Desktop"));
-	BuildPlatformInfo(TEXT("Linux"),					TEXT("Linux"),				LOCTEXT("LinuxEditor", "Linux (Editor)"),						EPlatformType::Editor,	EPlatformFlags::None,			FPlatformIconPaths(TEXT("Launcher/Linux/Platform_Linux_24x"), TEXT("Launcher/Linux/Platform_Linux_128x")),							TEXT(""),											TEXT("Linux_x64"),	EPlatformSDKStatus::Unknown,	TEXT("/Platforms/Linux/GettingStarted"),																IsAvailableOnLinux,												TEXT("Linux"),		TEXT("Linux"),		IsAvailableOnLinux,		true,					false,			TEXT("Linux"),		TEXT("Desktop"));
-	BuildPlatformInfo(TEXT("LinuxClient"),				TEXT("LinuxClient"),		LOCTEXT("LinuxClient", "Linux (Client-only)"),					EPlatformType::Client,	EPlatformFlags::None,			FPlatformIconPaths(TEXT("Launcher/Linux/Platform_Linux_24x"), TEXT("Launcher/Linux/Platform_Linux_128x")),							TEXT("-client"),									TEXT("Linux_x64"),	EPlatformSDKStatus::Unknown,	TEXT("/Platforms/Linux/GettingStarted"),																IsAvailableOnLinux || IsAvailableOnWindows,						TEXT("Linux"),		TEXT("Linux"),		IsAvailableOnLinux,		true,					false,			TEXT("Linux"),		TEXT("Desktop"));
-	BuildPlatformInfo(TEXT("LinuxServer"),				TEXT("LinuxServer"),		LOCTEXT("LinuxServer", "Linux (Dedicated Server)"),				EPlatformType::Server,	EPlatformFlags::None,			FPlatformIconPaths(TEXT("Launcher/Linux/Platform_Linux_24x"), TEXT("Launcher/Linux/Platform_Linux_128x")),							TEXT(""),											TEXT("Linux_x64"),	EPlatformSDKStatus::Unknown,	TEXT("/Platforms/Linux/GettingStarted"),																IsAvailableOnLinux || IsAvailableOnWindows,						TEXT("Linux"),		TEXT("Linux"),		IsAvailableOnLinux,		true,					false,			TEXT("Linux"),		TEXT("Desktop"));
-
-	BuildPlatformInfo(TEXT("IOS"),						TEXT("IOS"),				LOCTEXT("IOS", "iOS"),											EPlatformType::Game,	EPlatformFlags::None,			FPlatformIconPaths(TEXT("Launcher/iOS/Platform_iOS_24x"), TEXT("Launcher/iOS/Platform_iOS_128x")),									TEXT(""),											TEXT(""),			EPlatformSDKStatus::Unknown,	IOSTutorial,																							IsAvailableOnWindows || IsAvailableOnMac,						TEXT("IOS"),		TEXT("IOS"),		false,					true,					false,			TEXT("IOS"),		TEXT("Mobile"));
-	BuildPlatformInfo(TEXT("IOSClient"),				TEXT("IOSClient"),			LOCTEXT("IOSClient", "iOSClient"),								EPlatformType::Client,	EPlatformFlags::None,			FPlatformIconPaths(TEXT("Launcher/iOS/Platform_iOS_24x"), TEXT("Launcher/iOS/Platform_iOS_128x")),									TEXT("-client"),									TEXT(""),			EPlatformSDKStatus::Unknown,	IOSTutorial,																							IsAvailableOnWindows || IsAvailableOnMac,						TEXT("IOS"),		TEXT("IOS"),		false,					true,					false,			TEXT("IOS"),		TEXT("Mobile"));
-
-	BuildPlatformInfo(TEXT("Android"),					TEXT("Android"),			LOCTEXT("Android", "Android"),									EPlatformType::Game,	EPlatformFlags::None,			FPlatformIconPaths(TEXT("Launcher/Android/Platform_Android_24x"), TEXT("Launcher/Android/Platform_Android_128x")),					TEXT(""),											TEXT("Android"),	EPlatformSDKStatus::Unknown,	TEXT("/Engine/Tutorial/Mobile/SettingUpAndroidTutorial.SettingUpAndroidTutorial"),						IsAvailableOnLinux || IsAvailableOnWindows || IsAvailableOnMac,	TEXT("Android"),	TEXT("Android"),	false,					true,					false,			TEXT("Android"),	TEXT("Mobile"));
-	BuildPlatformInfo(TEXT("Android_Multi"),			TEXT("Android_Multi"),		LOCTEXT("Android_Multi", "Android (Multi)"),					EPlatformType::Game,	EPlatformFlags::CookFlavor,		FPlatformIconPaths(TEXT("Launcher/Android/Platform_Android_24x"), TEXT("Launcher/Android/Platform_Android_128x")),					TEXT("-targetplatform=Android -cookflavor=Multi"),	TEXT("Android"),	EPlatformSDKStatus::Unknown,	TEXT("/Engine/Tutorial/Mobile/SettingUpAndroidTutorial.SettingUpAndroidTutorial"),						IsAvailableOnLinux || IsAvailableOnWindows || IsAvailableOnMac,	TEXT("Android"),	TEXT("Android"),	false,					true,					false,			TEXT("Android"),	TEXT("Mobile"));
-	BuildPlatformInfo(TEXT("Android_ATC"),				TEXT("Android_ATC"),		LOCTEXT("Android_ATC", "Android (ATC)"),						EPlatformType::Game,	EPlatformFlags::CookFlavor,		FPlatformIconPaths(TEXT("Launcher/Android/Platform_Android_ATC_24x"), TEXT("Launcher/Android/Platform_Android_128x")),				TEXT("-targetplatform=Android -cookflavor=ATC"),	TEXT("Android"),	EPlatformSDKStatus::Unknown,	TEXT("/Engine/Tutorial/Mobile/SettingUpAndroidTutorial.SettingUpAndroidTutorial"),						IsAvailableOnLinux || IsAvailableOnWindows || IsAvailableOnMac,	TEXT("Android"),	TEXT("Android"),	false,					true,					false,			TEXT("Android"),	TEXT("Mobile"));
-	BuildPlatformInfo(TEXT("Android_DXT"),				TEXT("Android_DXT"),		LOCTEXT("Android_DXT", "Android (DXT)"),						EPlatformType::Game,	EPlatformFlags::CookFlavor,		FPlatformIconPaths(TEXT("Launcher/Android/Platform_Android_DXT_24x"), TEXT("Launcher/Android/Platform_Android_128x")),				TEXT("-targetplatform=Android -cookflavor=DXT"),	TEXT("Android"),	EPlatformSDKStatus::Unknown,	TEXT("/Engine/Tutorial/Mobile/SettingUpAndroidTutorial.SettingUpAndroidTutorial"),						IsAvailableOnLinux || IsAvailableOnWindows || IsAvailableOnMac,	TEXT("Android"),	TEXT("Android"),	false,					true,					false,			TEXT("Android"),	TEXT("Mobile"));
-	BuildPlatformInfo(TEXT("Android_ETC1"),				TEXT("Android_ETC1"),		LOCTEXT("Android_ETC1", "Android (ETC1)"),						EPlatformType::Game,	EPlatformFlags::CookFlavor,		FPlatformIconPaths(TEXT("Launcher/Android/Platform_Android_ETC1_24x"), TEXT("Launcher/Android/Platform_Android_128x")),				TEXT("-targetplatform=Android -cookflavor=ETC1"),	TEXT("Android"),	EPlatformSDKStatus::Unknown,	TEXT("/Engine/Tutorial/Mobile/SettingUpAndroidTutorial.SettingUpAndroidTutorial"),						IsAvailableOnLinux || IsAvailableOnWindows || IsAvailableOnMac,	TEXT("Android"),	TEXT("Android"),	false,					true,					false,			TEXT("Android"),	TEXT("Mobile"));
-	BuildPlatformInfo(TEXT("Android_ETC1a"),			TEXT("Android_ETC1a"),		LOCTEXT("Android_ETC1a", "Android (ETC1a)"),					EPlatformType::Game,	EPlatformFlags::CookFlavor,		FPlatformIconPaths(TEXT("Launcher/Android/Platform_Android_ETC1_24x"), TEXT("Launcher/Android/Platform_Android_128x")),				TEXT("-targetplatform=Android -cookflavor=ETC1a"),	TEXT("Android"),	EPlatformSDKStatus::Unknown,	TEXT("/Engine/Tutorial/Mobile/SettingUpAndroidTutorial.SettingUpAndroidTutorial"),						IsAvailableOnLinux || IsAvailableOnWindows || IsAvailableOnMac,	TEXT("Android"),	TEXT("Android"),	false,					true,					false,			TEXT("Android"),	TEXT("Mobile"));
-	BuildPlatformInfo(TEXT("Android_ETC2"),				TEXT("Android_ETC2"),		LOCTEXT("Android_ETC2", "Android (ETC2)"),						EPlatformType::Game,	EPlatformFlags::CookFlavor,		FPlatformIconPaths(TEXT("Launcher/Android/Platform_Android_ETC2_24x"), TEXT("Launcher/Android/Platform_Android_128x")),				TEXT("-targetplatform=Android -cookflavor=ETC2"),	TEXT("Android"),	EPlatformSDKStatus::Unknown,	TEXT("/Engine/Tutorial/Mobile/SettingUpAndroidTutorial.SettingUpAndroidTutorial"),						IsAvailableOnLinux || IsAvailableOnWindows || IsAvailableOnMac,	TEXT("Android"),	TEXT("Android"),	false,					true,					false,			TEXT("Android"),	TEXT("Mobile"));
-	BuildPlatformInfo(TEXT("Android_PVRTC"),			TEXT("Android_PVRTC"),		LOCTEXT("Android_PVRTC", "Android (PVRTC)"),					EPlatformType::Game,	EPlatformFlags::CookFlavor,		FPlatformIconPaths(TEXT("Launcher/Android/Platform_Android_PVRTC_24x"), TEXT("Launcher/Android/Platform_Android_128x")),			TEXT("-targetplatform=Android -cookflavor=PVRTC"),	TEXT("Android"),	EPlatformSDKStatus::Unknown,	TEXT("/Engine/Tutorial/Mobile/SettingUpAndroidTutorial.SettingUpAndroidTutorial"),						IsAvailableOnLinux || IsAvailableOnWindows || IsAvailableOnMac,	TEXT("Android"),	TEXT("Android"),	false,					true,					false,			TEXT("Android"),	TEXT("Mobile"));
-	BuildPlatformInfo(TEXT("Android_ASTC"),				TEXT("Android_ASTC"),		LOCTEXT("Android_ASTC", "Android (ASTC)"),						EPlatformType::Game,	EPlatformFlags::CookFlavor,		FPlatformIconPaths(TEXT("Launcher/Android/Platform_Android_ASTC_24x"), TEXT("Launcher/Android/Platform_Android_128x")),				TEXT("-targetplatform=Android -cookflavor=ASTC"),	TEXT("Android"),	EPlatformSDKStatus::Unknown,	TEXT("/Engine/Tutorial/Mobile/SettingUpAndroidTutorial.SettingUpAndroidTutorial"),						IsAvailableOnLinux || IsAvailableOnWindows || IsAvailableOnMac,	TEXT("Android"),	TEXT("Android"),	false,					true,					false,			TEXT("Android"),	TEXT("Mobile"));
-
-	BuildPlatformInfo(TEXT("AndroidClient"),			TEXT("AndroidClient"),		LOCTEXT("AndroidClient", "Android (Client-only)"),				EPlatformType::Client,	EPlatformFlags::None,			FPlatformIconPaths(TEXT("Launcher/Android/Platform_Android_24x"), TEXT("Launcher/Android/Platform_Android_128x")),					TEXT("-client"),									TEXT("Android"),	EPlatformSDKStatus::Unknown,	TEXT("/Engine/Tutorial/Mobile/SettingUpAndroidTutorial.SettingUpAndroidTutorial"),						IsAvailableOnLinux || IsAvailableOnWindows || IsAvailableOnMac,	TEXT("Android"),	TEXT("Android"),	false,					true,					false,			TEXT("Android"),	TEXT("Mobile"));
-	BuildPlatformInfo(TEXT("Android_MultiClient"),		TEXT("Android_MultiClient"),LOCTEXT("Android_MultiClient", "Android (Multi) (Client-only)"),EPlatformType::Client,	EPlatformFlags::CookFlavor,		FPlatformIconPaths(TEXT("Launcher/Android/Platform_Android_24x"), TEXT("Launcher/Android/Platform_Android_128x")),					TEXT("-client -targetplatform=Android -cookflavor=Multi"),	TEXT("Android"),	EPlatformSDKStatus::Unknown,	TEXT("/Engine/Tutorial/Mobile/SettingUpAndroidTutorial.SettingUpAndroidTutorial"),				IsAvailableOnLinux || IsAvailableOnWindows || IsAvailableOnMac,	TEXT("Android"),	TEXT("Android"),	false,					true,					false,			TEXT("Android"),	TEXT("Mobile"));
-	BuildPlatformInfo(TEXT("Android_ATCClient"),		TEXT("Android_ATCClient"),	LOCTEXT("Android_ATCClient", "Android (ATC) (Client-only)"),	EPlatformType::Client,	EPlatformFlags::CookFlavor,		FPlatformIconPaths(TEXT("Launcher/Android/Platform_Android_ATC_24x"), TEXT("Launcher/Android/Platform_Android_128x")),				TEXT("-client -targetplatform=Android -cookflavor=ATC"),	TEXT("Android"),	EPlatformSDKStatus::Unknown,	TEXT("/Engine/Tutorial/Mobile/SettingUpAndroidTutorial.SettingUpAndroidTutorial"),				IsAvailableOnLinux || IsAvailableOnWindows || IsAvailableOnMac,	TEXT("Android"),	TEXT("Android"),	false,					true,					false,			TEXT("Android"),	TEXT("Mobile"));
-	BuildPlatformInfo(TEXT("Android_DXTClient"),		TEXT("Android_DXTClient"),	LOCTEXT("Android_DXTClient", "Android (DXT) (Client-only)"),	EPlatformType::Client,	EPlatformFlags::CookFlavor,		FPlatformIconPaths(TEXT("Launcher/Android/Platform_Android_DXT_24x"), TEXT("Launcher/Android/Platform_Android_128x")),				TEXT("-client -targetplatform=Android -cookflavor=DXT"),	TEXT("Android"),	EPlatformSDKStatus::Unknown,	TEXT("/Engine/Tutorial/Mobile/SettingUpAndroidTutorial.SettingUpAndroidTutorial"),				IsAvailableOnLinux || IsAvailableOnWindows || IsAvailableOnMac,	TEXT("Android"),	TEXT("Android"),	false,					true,					false,			TEXT("Android"),	TEXT("Mobile"));
-	BuildPlatformInfo(TEXT("Android_ETC1Client"),		TEXT("Android_ETC1Client"),	LOCTEXT("Android_ETC1Client", "Android (ETC1) (Client-only)"),	EPlatformType::Client,	EPlatformFlags::CookFlavor,		FPlatformIconPaths(TEXT("Launcher/Android/Platform_Android_ETC1_24x"), TEXT("Launcher/Android/Platform_Android_128x")),				TEXT("-client -targetplatform=Android -cookflavor=ETC1"),	TEXT("Android"),	EPlatformSDKStatus::Unknown,	TEXT("/Engine/Tutorial/Mobile/SettingUpAndroidTutorial.SettingUpAndroidTutorial"),				IsAvailableOnLinux || IsAvailableOnWindows || IsAvailableOnMac,	TEXT("Android"),	TEXT("Android"),	false,					true,					false,			TEXT("Android"),	TEXT("Mobile"));
-	BuildPlatformInfo(TEXT("Android_ETC1aClient"),		TEXT("Android_ETC1aClient"),LOCTEXT("Android_ETC1aClient", "Android (ETC1a) (Client-only)"),EPlatformType::Client,	EPlatformFlags::CookFlavor,		FPlatformIconPaths(TEXT("Launcher/Android/Platform_Android_ETC1_24x"), TEXT("Launcher/Android/Platform_Android_128x")),				TEXT("-client -targetplatform=Android -cookflavor=ETC1a"),	TEXT("Android"),	EPlatformSDKStatus::Unknown,	TEXT("/Engine/Tutorial/Mobile/SettingUpAndroidTutorial.SettingUpAndroidTutorial"),				IsAvailableOnLinux || IsAvailableOnWindows || IsAvailableOnMac,	TEXT("Android"),	TEXT("Android"),	false,					true,					false,			TEXT("Android"),	TEXT("Mobile"));
-	BuildPlatformInfo(TEXT("Android_ETC2Client"),		TEXT("Android_ETC2Client"),	LOCTEXT("Android_ETC2Client", "Android (ETC2) (Client-only)"),	EPlatformType::Client,	EPlatformFlags::CookFlavor,		FPlatformIconPaths(TEXT("Launcher/Android/Platform_Android_ETC2_24x"), TEXT("Launcher/Android/Platform_Android_128x")),				TEXT("-client -targetplatform=Android -cookflavor=ETC2"),	TEXT("Android"),	EPlatformSDKStatus::Unknown,	TEXT("/Engine/Tutorial/Mobile/SettingUpAndroidTutorial.SettingUpAndroidTutorial"),				IsAvailableOnLinux || IsAvailableOnWindows || IsAvailableOnMac,	TEXT("Android"),	TEXT("Android"),	false,					true,					false,			TEXT("Android"),	TEXT("Mobile"));
-	BuildPlatformInfo(TEXT("Android_PVRTCClient"),		TEXT("Android_PVRTCClient"),LOCTEXT("Android_PVRTCClient", "Android (PVRTC) (Client-only)"),EPlatformType::Client,	EPlatformFlags::CookFlavor,		FPlatformIconPaths(TEXT("Launcher/Android/Platform_Android_PVRTC_24x"), TEXT("Launcher/Android/Platform_Android_128x")),			TEXT("-client -targetplatform=Android -cookflavor=PVRTC"),	TEXT("Android"),	EPlatformSDKStatus::Unknown,	TEXT("/Engine/Tutorial/Mobile/SettingUpAndroidTutorial.SettingUpAndroidTutorial"),				IsAvailableOnLinux || IsAvailableOnWindows || IsAvailableOnMac,	TEXT("Android"),	TEXT("Android"),	false,					true,					false,			TEXT("Android"),	TEXT("Mobile"));
-	BuildPlatformInfo(TEXT("Android_ASTCClient"),		TEXT("Android_ASTCClient"),	LOCTEXT("Android_ASTCClient", "Android (ASTC) (Client-only)"),	EPlatformType::Client,	EPlatformFlags::CookFlavor,		FPlatformIconPaths(TEXT("Launcher/Android/Platform_Android_ASTC_24x"), TEXT("Launcher/Android/Platform_Android_128x")),				TEXT("-client -targetplatform=Android -cookflavor=ASTC"),	TEXT("Android"),	EPlatformSDKStatus::Unknown,	TEXT("/Engine/Tutorial/Mobile/SettingUpAndroidTutorial.SettingUpAndroidTutorial"),				IsAvailableOnLinux || IsAvailableOnWindows || IsAvailableOnMac,	TEXT("Android"),	TEXT("Android"),	false,					true,					false,			TEXT("Android"),	TEXT("Mobile"));
-
-	BuildPlatformInfo(TEXT("HTML5"),					TEXT("HTML5"),				LOCTEXT("HTML5", "HTML5"),										EPlatformType::Game,	EPlatformFlags::None,			FPlatformIconPaths(TEXT("Launcher/HTML5/Platform_HTML5_24x"), TEXT("Launcher/HTML5/Platform_HTML5_128x")),							TEXT(""),											TEXT(""),			EPlatformSDKStatus::Unknown,	TEXT("/Platforms/HTML5/GettingStarted"),																IsAvailableOnLinux || IsAvailableOnWindows || IsAvailableOnMac,	TEXT("HTML5"),		TEXT("HTML5"),		false,					true,					false,			TEXT("HTML5"),		TEXT("Mobile"));
-
-	BuildPlatformInfo(TEXT("PS4"),						TEXT("PS4"),				LOCTEXT("PS4", "PlayStation 4"),								EPlatformType::Game,	EPlatformFlags::None,			FPlatformIconPaths(TEXT("Launcher/PS4/Platform_PS4_24x"), TEXT("Launcher/PS4/Platform_PS4_128x")),									TEXT(""),											TEXT("PS4"),		EPlatformSDKStatus::Unknown,	TEXT("/Platforms/PS4/GettingStarted"),																	IsAvailableOnWindows,											TEXT("PS4"),		TEXT("PS4"),		false,					false,					true,			TEXT("PS4"),		TEXT("Console"));
-
-	BuildPlatformInfo(TEXT("XboxOne"),					TEXT("XboxOne"),			LOCTEXT("XboxOne", "Xbox One"),									EPlatformType::Game,	EPlatformFlags::None,			FPlatformIconPaths(TEXT("Launcher/XboxOne/Platform_XboxOne_24x"), TEXT("Launcher/XboxOne/Platform_XboxOne_128x")),					TEXT(""),											TEXT("XboxOne"),	EPlatformSDKStatus::Unknown,	TEXT("/Platforms/XboxOne/GettingStarted"),																IsAvailableOnWindows,											TEXT("XboxOne"),	TEXT("XboxOne"),	false,					true,					true,			TEXT("XboxOne"),	TEXT("Console"));
-
-	BuildPlatformInfo(TEXT("AllDesktop"),				TEXT("AllDesktop"),			LOCTEXT("DesktopTargetPlatDisplay", "Desktop (Win+Mac+Linux)"),	EPlatformType::Game,	EPlatformFlags::None,			FPlatformIconPaths(TEXT("Launcher/Desktop/Platform_Desktop_24x"), TEXT("Launcher/Desktop/Platform_Desktop_128x")),					TEXT(""),											TEXT(""),			EPlatformSDKStatus::Unknown,	TEXT(""),																								IsAvailableOnWindows /* see note below */,						TEXT(""),			TEXT(""),			false,					true,					false,			TEXT("AllDesktop"),	TEXT("Desktop"));
-
-	BuildPlatformInfo(TEXT("TVOS"),						TEXT("TVOS"),				LOCTEXT("TVOSTargetPlatDisplay", "tvOS"),						EPlatformType::Game,	EPlatformFlags::None,			FPlatformIconPaths(TEXT("Launcher/TVOS/Platform_TVOS_24x"), TEXT("Launcher/TVOS/Platform_TVOS_128x")),								TEXT(""),											TEXT(""),			EPlatformSDKStatus::Unknown,	TEXT(""),																								IsAvailableOnWindows || IsAvailableOnMac,						TEXT("TVOS"),		TEXT("TVOS"),		false,					true,					false,			TEXT("TVOS"),		TEXT("Mobile"));
-	BuildPlatformInfo(TEXT("TVOSClient"),				TEXT("TVOSClient"),			LOCTEXT("TVOSTargetPlatDisplayClient", "tvOSClient"),			EPlatformType::Client,	EPlatformFlags::None,			FPlatformIconPaths(TEXT("Launcher/TVOS/Platform_TVOS_24x"), TEXT("Launcher/TVOS/Platform_TVOS_128x")),								TEXT("-client"),									TEXT(""),			EPlatformSDKStatus::Unknown,	TEXT(""),																								IsAvailableOnWindows || IsAvailableOnMac,						TEXT("TVOS"),		TEXT("TVOS"),		false,					true,					false,			TEXT("TVOS"),		TEXT("Mobile"));
-	BuildPlatformInfo(TEXT("Switch"),					TEXT("Switch"),				LOCTEXT("Switch", "Switch"),									EPlatformType::Game,	EPlatformFlags::None,			FPlatformIconPaths(TEXT("Launcher/Switch/Platform_Switch_24x"), TEXT("Launcher/Switch/Platform_Switch_128x")),						TEXT(""),											TEXT(""),			EPlatformSDKStatus::Unknown,	TEXT(""),																								IsAvailableOnWindows,											TEXT("Switch"),		TEXT("Switch"),		false,					true,					true,			TEXT("Switch"),		TEXT("Console"));
-	BuildPlatformInfo(TEXT("Lumin"),					TEXT("Lumin"),				LOCTEXT("Lumin", "Lumin"),										EPlatformType::Game,	EPlatformFlags::None,			FPlatformIconPaths(TEXT("Launcher/LuminTarget/Platform_Lumin_24x"), TEXT("Launcher/Lumin/Platform_Lumin_128x")),					TEXT(""),											TEXT(""),			EPlatformSDKStatus::Unknown,	TEXT("/Engine/Tutorial/Mobile/SettingUpLuminTutorial.SettingUpLuminTutorial"),							IsAvailableOnWindows || IsAvailableOnMac,						TEXT("Lumin"),		TEXT("Lumin"),		false,					true,					false,			TEXT("Lumin"),		TEXT("Mobile"));
-	BuildPlatformInfo(TEXT("LuminClient"),				TEXT("LuminClient"),		LOCTEXT("LuminClient", "Lumin (Client-only)"),					EPlatformType::Client,	EPlatformFlags::None,			FPlatformIconPaths(TEXT("Launcher/LuminTarget/Platform_Lumin_24x"), TEXT("Launcher/Lumin/Platform_Lumin_128x")),					TEXT("-client"),									TEXT(""),			EPlatformSDKStatus::Unknown,	TEXT("/Engine/Tutorial/Mobile/SettingUpLuminTutorial.SettingUpLuminTutorial"),							IsAvailableOnWindows || IsAvailableOnMac,						TEXT("Lumin"),		TEXT("Lumin"),		false,					true,					false,			TEXT("Lumin"),		TEXT("Mobile"));
-
-	BuildPlatformInfo(TEXT("Quail"),					TEXT("Quail"),				LOCTEXT("Quail", "Quail"),										EPlatformType::Game,	EPlatformFlags::None,			FPlatformIconPaths(TEXT("Launcher/Quail/Platform_Quail_24x"), TEXT("Launcher/Quail/Platform_Quail_128x")),							TEXT(""),											TEXT(""),			EPlatformSDKStatus::Unknown,	TEXT(""),																								IsAvailableOnWindows,											TEXT("Quail"),		TEXT("Quail"),		false,					false,					true,			TEXT("Quail"),		TEXT(""));
-	BuildPlatformInfo(TEXT("QuailClient"),				TEXT("QuailClient"),		LOCTEXT("QuailClient", "Quail (Client-only)"),					EPlatformType::Client,	EPlatformFlags::None,			FPlatformIconPaths(TEXT("Launcher/Quail/Platform_Quail_24x"), TEXT("Launcher/Quail/Platform_Quail_128x")),							TEXT("-client -targetplatform=Quail"),				TEXT(""),			EPlatformSDKStatus::Unknown,	TEXT(""),																								IsAvailableOnWindows,											TEXT("Quail"),		TEXT("Quail"),		false,					false,					true,			TEXT("Quail"),		TEXT(""));
-
+	BuildPlatformInfo(TEXT("AllDesktop"),				TEXT("AllDesktop"),			LOCTEXT("DesktopTargetPlatDisplay", "Desktop (Win+Mac+Linux)"),	EPlatformType::Game,	EPlatformFlags::None,			FPlatformIconPaths(TEXT("Launcher/Desktop/Platform_Desktop_24x"), TEXT("Launcher/Desktop/Platform_Desktop_128x")),					TEXT(""),											TEXT(""),			EPlatformSDKStatus::Unknown,	TEXT(""),																								PLATFORM_WINDOWS /* see note below */,						TEXT(""),			TEXT(""),			false,					true,					false,			TEXT("AllDesktop"),	TEXT("Desktop"));
 	// Note: For "AllDesktop" bEnabledForUse value, see SProjectTargetPlatformSettings::Construct !!!! IsAvailableOnWindows || IsAvailableOnMac || IsAvailableOnLinux
 }
 
 // Gets a string from a section, or empty string if it didn't exist
 FString GetSectionString(const FConfigSection& Section, FName Key)
 {
-	return Section.FindRef(Key).GetValue();
+	// look for a value prefixed with host:
+	FName HostKey = *FString::Printf(TEXT("%s:%s"), ANSI_TO_TCHAR(FPlatformProperties::IniPlatformName()), *Key.ToString());
+	const FConfigValue* HostValue = Section.Find(HostKey);
+	return HostValue ? HostValue->GetValue() : Section.FindRef(Key).GetValue();
 }
 
 // Gets a bool from a section, or false if it didn't exist
@@ -186,6 +112,7 @@ EPlatformFlags::Flags ConvertPlatformFlags(const FString& String)
 
 void ParseDataDrivenPlatformInfo(const TCHAR* Name, const FConfigSection& Section)
 {
+	checkf(WITH_EDITOR, TEXT("We don't expect to oarse DDPI on non-editor"));
 //	void BuildPlatformInfo(const FName& InPlatformInfoName, const FName& InTargetPlatformName, const FText& InDisplayName, const EPlatformType InPlatformType, const EPlatformFlags::Flags InPlatformFlags, 
 //		const FPlatformIconPaths& InIconPaths, const FString& InUATCommandLine, const FString& InAutoSDKPath, EPlatformSDKStatus InStatus, const FString& InTutorial, bool InEnabled, FString InBinaryFolderName, 
 //		FString InIniPlatformName, bool InUsesHostCompiler, bool InUATClosesAfterLaunch, bool InIsConfidential, const FName& InUBTTargetId, const FName& InPlatformGroupName)
@@ -200,19 +127,8 @@ void ParseDataDrivenPlatformInfo(const TCHAR* Name, const FConfigSection& Sectio
 	FString XLargeIconPath = GetSectionString(Section, TEXT("XLargeIconPath"));
 	FString UATCommandLine = GetSectionString(Section, TEXT("UATCommandLine"));
 	FString AutoSDKPath = GetSectionString(Section, TEXT("AutoSDKPath"));
-//	FString SDKStatus = GetSectionString(Section, TEXT("SDKStatus"));
 	FString TutorialPath = GetSectionString(Section, TEXT("TutorialPath"));
-#if PLATFORM_WINDOWS
-	bool bIsEnabled = GetSectionBool(Section, TEXT("bSupportsHostWindows"));
-#elif PLATFORM_MAC
-	bool bIsEnabled = GetSectionBool(Section, TEXT("bSupportsHostMac"));
-#elif PLATFORM_LINUX
-	bool bIsEnabled = GetSectionBool(Section, TEXT("bSupportsHostLinux"));
-#else
-	bool bIsEnabled = false;
-// @todo platplug: for some reason this is being compiled on target devices! TargetPlatform pulls in DesktopPlatform, but can't figure out why TP is being pulled in
-//#error Unknown Host Platform!
-#endif
+	bool bIsEnabled = GetSectionBool(Section, TEXT("bIsEnabled"));
 	FString BinariesDirectoryName = GetSectionString(Section, TEXT("BinariesDirectoryName"));
 	FString IniPlatformName = GetSectionString(Section, TEXT("IniPlatformName"));
 	bool bUsesHostCompiler = GetSectionBool(Section, TEXT("bUsesHostCompiler"));
@@ -303,59 +219,18 @@ void UpdatePlatformDisplayName(FString InPlatformName, FText InDisplayName)
 	}
 }
 
-const TArray<FPlatformInfo>& GetPlatformInfoArray(bool bAccessiblePlatformsOnly)
+const TArray<FPlatformInfo>& GetPlatformInfoArray()
 {
-	if (bAccessiblePlatformsOnly)
-	{
-		static bool bHasSearchedForPlatforms = false;
-		static TArray<FPlatformInfo> AccessiblePlatforms;
-
-		if (!bHasSearchedForPlatforms)
-		{
-			const TArray<FString>& ConfidentalPlatforms = FDataDrivenPlatformInfoRegistry::GetConfidentialPlatforms();
-
-			for (const FPlatformInfo& PlatformInfo : AllPlatformInfoArray)
-			{
-				TArray<FString> FoundFiles;
-				const FString IniFolderPath = FPaths::Combine(*FPaths::EngineConfigDir(), *PlatformInfo.IniPlatformName);
-				const FString ExtIniFolderPath = FPaths::Combine(*FPaths::PlatformExtensionsDir(), *PlatformInfo.IniPlatformName, TEXT("Engine"), TEXT("Config"));
-				IFileManager::Get().FindFiles(FoundFiles, *IniFolderPath);
-				IFileManager::Get().FindFiles(FoundFiles, *ExtIniFolderPath);
-
-				// We check that the configuration directory exists with actual files in it to include the platform
-				// P4/Git may have filtered out platforms and we don't want to include filtered platform to keep code from trying to load other files that don't exist.
-				if (FoundFiles.Num() > 0)
-				{
-					if (PlatformInfo.bIsConfidential && ConfidentalPlatforms.Contains(PlatformInfo.IniPlatformName))
-					{
-						AccessiblePlatforms.Add(PlatformInfo);
-					}
-					else if (!PlatformInfo.bIsConfidential)
-					{
-						AccessiblePlatforms.Add(PlatformInfo);
-					}
-				}
-			}
-
-			bHasSearchedForPlatforms = true;
-		}
-
-		return AccessiblePlatforms;
-	}
-	else
-	{
-
-		return AllPlatformInfoArray;
-	}
+	return AllPlatformInfoArray;
 }
 
-TArray<FVanillaPlatformEntry> BuildPlatformHierarchy(const EPlatformFilter InFilter, bool bAccessiblePlatformsOnly)
+TArray<FVanillaPlatformEntry> BuildPlatformHierarchy(const EPlatformFilter InFilter)
 {
 	TArray<FVanillaPlatformEntry> VanillaPlatforms;
 
 	// Build up a tree from the platforms we support (vanilla outers, with a list of flavors)
 	// PlatformInfoArray should be ordered in such a way that the vanilla platforms always appear before their flavors
-	for (const PlatformInfo::FPlatformInfo& PlatformInfo : GetPlatformInfoArray(bAccessiblePlatformsOnly))
+	for (const PlatformInfo::FPlatformInfo& PlatformInfo : AllPlatformInfoArray)
 	{
 		if(PlatformInfo.IsVanilla())
 		{
@@ -387,7 +262,7 @@ TArray<FVanillaPlatformEntry> BuildPlatformHierarchy(const EPlatformFilter InFil
 	return VanillaPlatforms;
 }
 
-FVanillaPlatformEntry BuildPlatformHierarchy(const FName& InPlatformName, const EPlatformFilter InFilter, bool bAccessiblePlatformsOnly)
+FVanillaPlatformEntry BuildPlatformHierarchy(const FName& InPlatformName, const EPlatformFilter InFilter)
 {
 	FVanillaPlatformEntry VanillaPlatformEntry;
 	const FPlatformInfo* VanillaPlatformInfo = FindVanillaPlatformInfo(InPlatformName);
@@ -396,7 +271,7 @@ FVanillaPlatformEntry BuildPlatformHierarchy(const FName& InPlatformName, const
 	{
 		VanillaPlatformEntry.PlatformInfo = VanillaPlatformInfo;
 		
-		for (const PlatformInfo::FPlatformInfo& PlatformInfo : GetPlatformInfoArray(bAccessiblePlatformsOnly))
+		for (const PlatformInfo::FPlatformInfo& PlatformInfo : AllPlatformInfoArray)
 		{
 			if (!PlatformInfo.IsVanilla() && PlatformInfo.VanillaPlatformName == VanillaPlatformInfo->PlatformInfoName)
 			{
diff --git a/Engine/Source/Developer/DesktopPlatform/Public/PlatformInfo.h b/Engine/Source/Developer/DesktopPlatform/Public/PlatformInfo.h
--- a/Engine/Source/Developer/DesktopPlatform/Public/PlatformInfo.h
+++ b/Engine/Source/Developer/DesktopPlatform/Public/PlatformInfo.h
@@ -264,24 +264,22 @@ namespace PlatformInfo
 	 * Get an array of all the platforms we know about
 	 * @return The pointer to the start of the platform array
 	 */
-	DESKTOPPLATFORM_API const TArray<FPlatformInfo>& GetPlatformInfoArray(bool bAccessiblePlatformsOnly = true);
+	DESKTOPPLATFORM_API const TArray<FPlatformInfo>& GetPlatformInfoArray();
 
 	/**
 	 * Build a hierarchy mapping vanilla platforms to their flavors
 	 * @param InFilter - Flags to control which kinds of flavors you want to include
-	 * @param bAccessiblePlatformsOnly	If true, only the accessible platforms(installed, or could be installed) will be returned
 	 * @return An array of vanilla platforms, potentially containing flavors
 	 */
-	DESKTOPPLATFORM_API TArray<FVanillaPlatformEntry> BuildPlatformHierarchy(const EPlatformFilter InFilter, bool bAccessiblePlatformsOnly = true);
+	DESKTOPPLATFORM_API TArray<FVanillaPlatformEntry> BuildPlatformHierarchy(const EPlatformFilter InFilter);
 
 	/**
 	* Build a hierarchy mapping for specified vanilla platform to it flavors
 	* @param InPlatformName - Platform name to build hierarchy for, could be vanilla or flavor name
 	* @param InFilter - Flags to control which kinds of flavors you want to include
-	* @param bAccessiblePlatformsOnly	If true, only the accessible platforms(installed, or could be installed) will be returned
 	* @return Vanilla platform potentially containing flavors
 	*/
-	DESKTOPPLATFORM_API FVanillaPlatformEntry BuildPlatformHierarchy(const FName& InPlatformName, const EPlatformFilter InFilter, bool bAccessiblePlatformsOnly = true);
+	DESKTOPPLATFORM_API FVanillaPlatformEntry BuildPlatformHierarchy(const FName& InPlatformName, const EPlatformFilter InFilter);
 
 	/**
 	 * Get an array of all the platforms we know about
diff --git a/Engine/Source/Developer/Localization/Private/LocTextHelper.cpp b/Engine/Source/Developer/Localization/Private/LocTextHelper.cpp
--- a/Engine/Source/Developer/Localization/Private/LocTextHelper.cpp
+++ b/Engine/Source/Developer/Localization/Private/LocTextHelper.cpp
@@ -35,7 +35,7 @@ const TArray<FString>& FLocTextPlatformSplitUtils::GetPlatformsToSplit(const ELo
 			static TArray<FString> AllPlatformNames = []()
 			{
 				TArray<FString> TmpArray;
-				for (const PlatformInfo::FPlatformInfo& Info : PlatformInfo::GetPlatformInfoArray(false))
+				for (const PlatformInfo::FPlatformInfo& Info : PlatformInfo::GetPlatformInfoArray())
 				{
 					if (!Info.IniPlatformName.IsEmpty())
 					{
diff --git a/Engine/Source/Developer/TargetPlatform/Private/TargetPlatformManagerModule.cpp b/Engine/Source/Developer/TargetPlatform/Private/TargetPlatformManagerModule.cpp
--- a/Engine/Source/Developer/TargetPlatform/Private/TargetPlatformManagerModule.cpp
+++ b/Engine/Source/Developer/TargetPlatform/Private/TargetPlatformManagerModule.cpp
@@ -103,7 +103,7 @@ public:
 
 			// we have to setup our local environment according to AutoSDKs or the ITargetPlatform's IsSDkInstalled calls may fail
 			// before we get a change to setup for a given platform.  Use the platforminfo list to avoid any kind of interdependency.
-			for (const PlatformInfo::FPlatformInfo& PlatformInfo : PlatformInfo::GetPlatformInfoArray(false))
+			for (const PlatformInfo::FPlatformInfo& PlatformInfo : PlatformInfo::GetPlatformInfoArray())
 			{
 				if (PlatformInfo.AutoSDKPath.Len() > 0)
 				{
diff --git a/Engine/Source/Editor/UnrealEd/Private/Commandlets/DiffAssetRegistriesCommandlet.cpp b/Engine/Source/Editor/UnrealEd/Private/Commandlets/DiffAssetRegistriesCommandlet.cpp
--- a/Engine/Source/Editor/UnrealEd/Private/Commandlets/DiffAssetRegistriesCommandlet.cpp
+++ b/Engine/Source/Editor/UnrealEd/Private/Commandlets/DiffAssetRegistriesCommandlet.cpp
@@ -248,7 +248,7 @@ int32 UDiffAssetRegistriesCommandlet::Main(const FString& FullCommandLine)
 		FString NewPlatformDir = NewPath.Left(NewPath.Find(AssetRegistrySubPath));
 		FString PlatformPath = FPaths::GetCleanFilename(NewPlatformDir);
 
-		for (const FPlatformInfo& PlatformInfo : PlatformInfo::GetPlatformInfoArray(false))
+		for (const FPlatformInfo& PlatformInfo : PlatformInfo::GetPlatformInfoArray())
 		{
 			if (PlatformPath == PlatformInfo.TargetPlatformName.ToString())
 			{
diff --git a/Engine/Source/Runtime/Core/Private/Misc/DataDrivenPlatformInfoRegistry.cpp b/Engine/Source/Runtime/Core/Private/Misc/DataDrivenPlatformInfoRegistry.cpp
--- a/Engine/Source/Runtime/Core/Private/Misc/DataDrivenPlatformInfoRegistry.cpp
+++ b/Engine/Source/Runtime/Core/Private/Misc/DataDrivenPlatformInfoRegistry.cpp
@@ -38,7 +38,7 @@ bool FDataDrivenPlatformInfoRegistry::LoadDataDrivenIniFile(int32 Index, FConfig
 
 	FConfigCacheIni::LoadExternalIniFile(IniFile, *FPaths::GetBaseFilename(IniFilenames[Index]), nullptr, *FPaths::GetPath(IniFilenames[Index]), false);
 
-	// plugin paths are different (platform/engine/config, not engine/config/platform)
+	// platform extension paths are different (platform/engine/config, not engine/config/platform)
 	if (IniFilenames[Index].StartsWith(FPaths::PlatformExtensionsDir()))
 	{
 		PlatformName = FPaths::GetCleanFilename(FPaths::GetPath(FPaths::GetPath(FPaths::GetPath(IniFilenames[Index]))));
