- Fixed XLarge icons for PlatformInfo

diff --git a/Engine/Source/Developer/DesktopPlatform/Private/PlatformInfo.cpp b/Engine/Source/Developer/DesktopPlatform/Private/PlatformInfo.cpp
--- a/Engine/Source/Developer/DesktopPlatform/Private/PlatformInfo.cpp
+++ b/Engine/Source/Developer/DesktopPlatform/Private/PlatformInfo.cpp
@@ -68,10 +68,7 @@ void BuildPlatformInfo(const FName& InPlatformInfoName, const FName& InTargetPla
 	// Generate the icon style names for FEditorStyle
 	PlatformInfo.IconPaths.NormalStyleName = *FString::Printf(TEXT("Launcher.Platform_%s"), *InPlatformInfoNameString);
 	PlatformInfo.IconPaths.LargeStyleName  = *FString::Printf(TEXT("Launcher.Platform_%s.Large"), *InPlatformInfoNameString);
-	if (PlatformInfo.IconPaths.XLargePath != TEXT(""))
-	{
-		PlatformInfo.IconPaths.XLargeStyleName = *FString::Printf(TEXT("Launcher.Platform_%s.XLarge"), *InPlatformInfoNameString);
-	}
+	PlatformInfo.IconPaths.XLargeStyleName = *FString::Printf(TEXT("Launcher.Platform_%s.XLarge"), *InPlatformInfoNameString);
 
 	// SDK data
 	PlatformInfo.SDKStatus = InStatus;
@@ -127,6 +124,11 @@ void ParseDataDrivenPlatformInfo(const TCHAR* Name, const FConfigSection& Sectio
 	FString NormalIconPath = GetSectionString(Section, TEXT("NormalIconPath"));
 	FString LargeIconPath = GetSectionString(Section, TEXT("LargeIconPath"));
 	FString XLargeIconPath = GetSectionString(Section, TEXT("XLargeIconPath"));
+	// no one has an XLarge path yet, but in case they add one, this will use it
+	if (XLargeIconPath == TEXT(""))
+	{
+		XLargeIconPath = LargeIconPath;
+	}
 	FString UATCommandLine = GetSectionString(Section, TEXT("UATCommandLine"));
 	FString AutoSDKPath = GetSectionString(Section, TEXT("AutoSDKPath"));
 	FString TutorialPath = GetSectionString(Section, TEXT("TutorialPath"));
