platform extensions - fix for UnrealFontend.

diff --git a/Engine/Source/Developer/DesktopPlatform/Private/PlatformInfo.cpp b/Engine/Source/Developer/DesktopPlatform/Private/PlatformInfo.cpp
--- a/Engine/Source/Developer/DesktopPlatform/Private/PlatformInfo.cpp
+++ b/Engine/Source/Developer/DesktopPlatform/Private/PlatformInfo.cpp
@@ -20,7 +20,7 @@ TArray<FPlatformInfo> AllPlatformInfoArray;
 
 // we don't need any of this without the editor, although we would ideally not even compile this outside of the editor
 // @todo platplug: Figure out why this is compiled on target devices
-#if WITH_EDITOR
+#if WITH_EDITOR || IS_PROGRAM
 
 void BuildPlatformInfo(const FName& InPlatformInfoName, const FName& InTargetPlatformName, const FText& InDisplayName, const EPlatformType InPlatformType, const EPlatformFlags::Flags InPlatformFlags, const FPlatformIconPaths& InIconPaths, const FString& InUATCommandLine, const FString& InAutoSDKPath, EPlatformSDKStatus InStatus, const FString& InTutorial, bool InEnabled, FString InBinaryFolderName, FString InIniPlatformName, bool InUsesHostCompiler, bool InUATClosesAfterLaunch, bool InIsConfidential, const FName& InUBTTargetId, const FName& InPlatformGroupName, bool InTargetPlatformCanUseCrashReporter)
 {
diff --git a/Engine/Source/Programs/UnrealFrontend/Private/Commands/UserInterfaceCommand.cpp b/Engine/Source/Programs/UnrealFrontend/Private/Commands/UserInterfaceCommand.cpp
--- a/Engine/Source/Programs/UnrealFrontend/Private/Commands/UserInterfaceCommand.cpp
+++ b/Engine/Source/Programs/UnrealFrontend/Private/Commands/UserInterfaceCommand.cpp
@@ -35,6 +35,8 @@ void FUserInterfaceCommand::Run(  )
 	FString UnrealFrontendLayoutIni = FPaths::GetPath(GEngineIni) + "/Layout.ini";
 
 	// ensure target platform manager is referenced early as it must be created on the main thread
+	FModuleManager::Get().LoadModuleChecked("DesktopPlatform");
+	FConfigCacheIni::InitializeConfigSystem();
 	GetTargetPlatformManager();
 
 	FCoreStyle::ResetToDefault();
