Merging CL 4336998 from Dev-VR - Splitting game and client platform info

diff --git a/Engine/Source/Developer/DesktopPlatform/Private/PlatformInfo.cpp b/Engine/Source/Developer/DesktopPlatform/Private/PlatformInfo.cpp
--- a/Engine/Source/Developer/DesktopPlatform/Private/PlatformInfo.cpp
+++ b/Engine/Source/Developer/DesktopPlatform/Private/PlatformInfo.cpp
@@ -105,7 +105,7 @@ static const FPlatformInfo AllPlatformInfoArray[] = {
 	BuildPlatformInfo(TEXT("LinuxServer"),				TEXT("LinuxServer"),		LOCTEXT("LinuxServer", "Linux (Dedicated Server)"),				EPlatformType::Server,	EPlatformFlags::None,			FPlatformIconPaths(TEXT("Launcher/Linux/Platform_Linux_24x"), TEXT("Launcher/Linux/Platform_Linux_128x")),							TEXT(""),											TEXT("Linux_x64"),	EPlatformSDKStatus::Unknown,	TEXT("/Platforms/Linux/GettingStarted"),																IsAvailableOnLinux || IsAvailableOnWindows,						TEXT("Linux"),		TEXT("Linux"),		IsAvailableOnLinux,		true,					false,			TEXT("Linux"),		TEXT("Desktop")),
 
 	BuildPlatformInfo(TEXT("IOS"),						TEXT("IOS"),				LOCTEXT("IOS", "iOS"),											EPlatformType::Game,	EPlatformFlags::None,			FPlatformIconPaths(TEXT("Launcher/iOS/Platform_iOS_24x"), TEXT("Launcher/iOS/Platform_iOS_128x")),									TEXT(""),											TEXT(""),			EPlatformSDKStatus::Unknown,	IOSTutorial,																							IsAvailableOnWindows || IsAvailableOnMac,						TEXT("IOS"),		TEXT("IOS"),		false,					true,					false,			TEXT("IOS"),		TEXT("Mobile")),
-	BuildPlatformInfo(TEXT("IOSClient"),				TEXT("IOSClient"),			LOCTEXT("IOSClient", "iOSClient"),								EPlatformType::Client,	EPlatformFlags::None,			FPlatformIconPaths(TEXT("Launcher/iOS/Platform_iOS_24x"), TEXT("Launcher/iOS/Platform_iOS_128x")),									TEXT("-client"),									TEXT(""),			EPlatformSDKStatus::Unknown,	IOSTutorial,																							IsAvailableOnWindows || IsAvailableOnMac,						TEXT("IOS"),		TEXT("IOS"),		false,					true,					false,			TEXT("IOS"),		TEXT("Mobile")),
+	BuildPlatformInfo(TEXT("IOSClient"),				TEXT("IOSClient"),			LOCTEXT("IOSClient", "iOS (Client-only)"),						EPlatformType::Client,	EPlatformFlags::None,			FPlatformIconPaths(TEXT("Launcher/iOS/Platform_iOS_24x"), TEXT("Launcher/iOS/Platform_iOS_128x")),									TEXT("-client"),									TEXT(""),			EPlatformSDKStatus::Unknown,	IOSTutorial,																							IsAvailableOnWindows || IsAvailableOnMac,						TEXT("IOS"),		TEXT("IOS"),		false,					true,					false,			TEXT("IOS"),		TEXT("Mobile")),
 
 	BuildPlatformInfo(TEXT("Android"),					TEXT("Android"),			LOCTEXT("Android", "Android"),									EPlatformType::Game,	EPlatformFlags::None,			FPlatformIconPaths(TEXT("Launcher/Android/Platform_Android_24x"), TEXT("Launcher/Android/Platform_Android_128x")),					TEXT(""),											TEXT("Android"),	EPlatformSDKStatus::Unknown,	TEXT("/Engine/Tutorial/Mobile/SettingUpAndroidTutorial.SettingUpAndroidTutorial"),						IsAvailableOnLinux || IsAvailableOnWindows || IsAvailableOnMac,	TEXT("Android"),	TEXT("Android"),	false,					true,					false,			TEXT("Android"),	TEXT("Mobile")),
 	BuildPlatformInfo(TEXT("Android_Multi"),			TEXT("Android_Multi"),		LOCTEXT("Android_Multi", "Android (Multi)"),					EPlatformType::Game,	EPlatformFlags::CookFlavor,		FPlatformIconPaths(TEXT("Launcher/Android/Platform_Android_24x"), TEXT("Launcher/Android/Platform_Android_128x")),					TEXT("-targetplatform=Android -cookflavor=Multi"),	TEXT("Android"),	EPlatformSDKStatus::Unknown,	TEXT("/Engine/Tutorial/Mobile/SettingUpAndroidTutorial.SettingUpAndroidTutorial"),						IsAvailableOnLinux || IsAvailableOnWindows || IsAvailableOnMac,	TEXT("Android"),	TEXT("Android"),	false,					true,					false,			TEXT("Android"),	TEXT("Mobile")),
@@ -130,16 +130,21 @@ static const FPlatformInfo AllPlatformInfoArray[] = {
 	BuildPlatformInfo(TEXT("HTML5"),					TEXT("HTML5"),				LOCTEXT("HTML5", "HTML5"),										EPlatformType::Game,	EPlatformFlags::None,			FPlatformIconPaths(TEXT("Launcher/HTML5/Platform_HTML5_24x"), TEXT("Launcher/HTML5/Platform_HTML5_128x")),							TEXT(""),											TEXT(""),			EPlatformSDKStatus::Unknown,	TEXT("/Platforms/HTML5/GettingStarted"),																IsAvailableOnLinux || IsAvailableOnWindows || IsAvailableOnMac,	TEXT("HTML5"),		TEXT("HTML5"),		false,					true,					false,			TEXT("HTML5"),		TEXT("Mobile")),
 
 	BuildPlatformInfo(TEXT("PS4"),						TEXT("PS4"),				LOCTEXT("PS4", "PlayStation 4"),								EPlatformType::Game,	EPlatformFlags::None,			FPlatformIconPaths(TEXT("Launcher/PS4/Platform_PS4_24x"), TEXT("Launcher/PS4/Platform_PS4_128x")),									TEXT(""),											TEXT("PS4"),		EPlatformSDKStatus::Unknown,	TEXT("/Platforms/PS4/GettingStarted"),																	IsAvailableOnWindows,											TEXT("PS4"),		TEXT("PS4"),		false,					false,					true,			TEXT("PS4"),		TEXT("Console")),
+	BuildPlatformInfo(TEXT("PS4Client"),				TEXT("PS4Client"),			LOCTEXT("PS4Client", "PlayStation 4 (Client-only)"),			EPlatformType::Client,	EPlatformFlags::None,			FPlatformIconPaths(TEXT("Launcher/PS4/Platform_PS4_24x"), TEXT("Launcher/PS4/Platform_PS4_128x")),									TEXT("-client"),									TEXT("PS4"),		EPlatformSDKStatus::Unknown,	TEXT("/Platforms/PS4/GettingStarted"),																	IsAvailableOnWindows,											TEXT("PS4"),		TEXT("PS4"),		false,					false,					true,			TEXT("PS4"),		TEXT("Console")),
 
 	BuildPlatformInfo(TEXT("XboxOne"),					TEXT("XboxOne"),			LOCTEXT("XboxOne", "Xbox One"),									EPlatformType::Game,	EPlatformFlags::None,			FPlatformIconPaths(TEXT("Launcher/XboxOne/Platform_XboxOne_24x"), TEXT("Launcher/XboxOne/Platform_XboxOne_128x")),					TEXT(""),											TEXT("XboxOne"),	EPlatformSDKStatus::Unknown,	TEXT("/Platforms/XboxOne/GettingStarted"),																IsAvailableOnWindows,											TEXT("XboxOne"),	TEXT("XboxOne"),	false,					true,					true,			TEXT("XboxOne"),	TEXT("Console")),
+	BuildPlatformInfo(TEXT("XboxOneClient"),			TEXT("XboxOneClient"),		LOCTEXT("XboxOneClient", "Xbox One (Client-only)"),				EPlatformType::Client,	EPlatformFlags::None,			FPlatformIconPaths(TEXT("Launcher/XboxOne/Platform_XboxOne_24x"), TEXT("Launcher/XboxOne/Platform_XboxOne_128x")),					TEXT("-client"),									TEXT("XboxOne"),	EPlatformSDKStatus::Unknown,	TEXT("/Platforms/XboxOne/GettingStarted"),																IsAvailableOnWindows,											TEXT("XboxOne"),	TEXT("XboxOne"),	false,					true,					true,			TEXT("XboxOne"),	TEXT("Console")),
 
 	BuildPlatformInfo(TEXT("AllDesktop"),				TEXT("AllDesktop"),			LOCTEXT("DesktopTargetPlatDisplay", "Desktop (Win+Mac+Linux)"),	EPlatformType::Game,	EPlatformFlags::None,			FPlatformIconPaths(TEXT("Launcher/Desktop/Platform_Desktop_24x"), TEXT("Launcher/Desktop/Platform_Desktop_128x")),					TEXT(""),											TEXT(""),			EPlatformSDKStatus::Unknown,	TEXT(""),																								IsAvailableOnWindows /* see note below */,						TEXT(""),			TEXT(""),			false,					true,					false,			TEXT("AllDesktop"),	TEXT("Desktop")),
 
 	BuildPlatformInfo(TEXT("TVOS"),						TEXT("TVOS"),				LOCTEXT("TVOSTargetPlatDisplay", "tvOS"),						EPlatformType::Game,	EPlatformFlags::None,			FPlatformIconPaths(TEXT("Launcher/TVOS/Platform_TVOS_24x"), TEXT("Launcher/TVOS/Platform_TVOS_128x")),								TEXT(""),											TEXT(""),			EPlatformSDKStatus::Unknown,	TEXT(""),																								IsAvailableOnWindows || IsAvailableOnMac,						TEXT("TVOS"),		TEXT("TVOS"),		false,					true,					false,			TEXT("TVOS"),		TEXT("Mobile")),
 	BuildPlatformInfo(TEXT("TVOSClient"),				TEXT("TVOSClient"),			LOCTEXT("TVOSTargetPlatDisplayClient", "tvOSClient"),			EPlatformType::Client,	EPlatformFlags::None,			FPlatformIconPaths(TEXT("Launcher/TVOS/Platform_TVOS_24x"), TEXT("Launcher/TVOS/Platform_TVOS_128x")),								TEXT("-client"),									TEXT(""),			EPlatformSDKStatus::Unknown,	TEXT(""),																								IsAvailableOnWindows || IsAvailableOnMac,						TEXT("TVOS"),		TEXT("TVOS"),		false,					true,					false,			TEXT("TVOS"),		TEXT("Mobile")),
 	BuildPlatformInfo(TEXT("Switch"),					TEXT("Switch"),				LOCTEXT("Switch", "Switch"),									EPlatformType::Game,	EPlatformFlags::None,			FPlatformIconPaths(TEXT("Launcher/Switch/Platform_Switch_24x"), TEXT("Launcher/Switch/Platform_Switch_128x")),						TEXT(""),											TEXT(""),			EPlatformSDKStatus::Unknown,	TEXT(""),																								IsAvailableOnWindows,											TEXT("Switch"),		TEXT("Switch"),		false,					true,					true,			TEXT("Switch"),		TEXT("Console")),
+	BuildPlatformInfo(TEXT("SwitchClient"),				TEXT("SwitchClient"),		LOCTEXT("SwitchClient", "Switch (Client-only)"),				EPlatformType::Client,	EPlatformFlags::None,			FPlatformIconPaths(TEXT("Launcher/Switch/Platform_Switch_24x"), TEXT("Launcher/Switch/Platform_Switch_128x")),						TEXT("-client"),									TEXT(""),			EPlatformSDKStatus::Unknown,	TEXT(""),																								IsAvailableOnWindows,											TEXT("Switch"),		TEXT("Switch"),		false,					true,					true,			TEXT("Switch"),		TEXT("Console")),
+
 	BuildPlatformInfo(TEXT("Lumin"),					TEXT("Lumin"),				LOCTEXT("Lumin", "Lumin"),										EPlatformType::Game,	EPlatformFlags::None,			FPlatformIconPaths(TEXT("Launcher/LuminTarget/Platform_Lumin_24x"), TEXT("Launcher/Lumin/Platform_Lumin_128x")),					TEXT(""),											TEXT(""),			EPlatformSDKStatus::Unknown,	TEXT("/Engine/Tutorial/Mobile/SettingUpLuminTutorial.SettingUpLuminTutorial"),							IsAvailableOnWindows || IsAvailableOnMac,						TEXT("Lumin"),		TEXT("Lumin"),		false,					true,					false,			TEXT("Lumin"),		TEXT("Mobile")),
-	
+	BuildPlatformInfo(TEXT("LuminClient"),				TEXT("LuminClient"),		LOCTEXT("LuminClient", "Lumin (Client-only)"),					EPlatformType::Client,	EPlatformFlags::None,			FPlatformIconPaths(TEXT("Launcher/LuminTarget/Platform_Lumin_24x"), TEXT("Launcher/Lumin/Platform_Lumin_128x")),					TEXT("-client"),									TEXT(""),			EPlatformSDKStatus::Unknown,	TEXT("/Engine/Tutorial/Mobile/SettingUpLuminTutorial.SettingUpLuminTutorial"),							IsAvailableOnWindows || IsAvailableOnMac,						TEXT("Lumin"),		TEXT("Lumin"),		false,					true,					false,			TEXT("Lumin"),		TEXT("Mobile")),
+
 	BuildPlatformInfo(TEXT("Quail"),					TEXT("Quail"),				LOCTEXT("Quail", "Quail"),										EPlatformType::Game,	EPlatformFlags::None,			FPlatformIconPaths(TEXT("Launcher/Quail/Platform_Quail_24x"), TEXT("Launcher/Quail/Platform_Quail_128x")),							TEXT(""),											TEXT(""),			EPlatformSDKStatus::Unknown,	TEXT(""),																								IsAvailableOnWindows,											TEXT("Quail"),		TEXT("Quail"),		false,					true,					true,			TEXT("Quail"),		TEXT("")),
 	BuildPlatformInfo(TEXT("QuailClient"),				TEXT("QuailClient"),		LOCTEXT("QuailClient", "Quail (Client-only)"),					EPlatformType::Client,	EPlatformFlags::None,			FPlatformIconPaths(TEXT("Launcher/Quail/Platform_Quail_24x"), TEXT("Launcher/Quail/Platform_Quail_128x")),							TEXT("-client -targetplatform=Quail"),				TEXT(""),			EPlatformSDKStatus::Unknown,	TEXT(""),																								IsAvailableOnWindows,											TEXT("Quail"),		TEXT("Quail"),		false,					true,					true,			TEXT("Quail"),		TEXT("")),
 
diff --git a/Engine/Source/Developer/Lumin/LuminTargetPlatform/Private/LuminTargetPlatform.cpp b/Engine/Source/Developer/Lumin/LuminTargetPlatform/Private/LuminTargetPlatform.cpp
--- a/Engine/Source/Developer/Lumin/LuminTargetPlatform/Private/LuminTargetPlatform.cpp
+++ b/Engine/Source/Developer/Lumin/LuminTargetPlatform/Private/LuminTargetPlatform.cpp
@@ -27,7 +27,7 @@ FLuminTargetPlatform::FLuminTargetPlatform(bool bIsClient)
 {
 #if WITH_ENGINE
 	// by using the FAndroidPlatformProperties, the PlatformInfo up in TTargetPlatformBase/FTargetPlatformBase would be Android
-	this->PlatformInfo = PlatformInfo::FindPlatformInfo("Lumin");
+	this->PlatformInfo = PlatformInfo::FindPlatformInfo(FName(*PlatformName()));
 
 	RefreshSettings();
 #endif
@@ -86,12 +86,16 @@ void FLuminTargetPlatform::RefreshSettings()
 #if WITH_ENGINE
 	//UE_LOG(LogCore, Warning, TEXT("*** DIAGNOSE - REFRESH!!!"));
 
+	//the load above does not move settings from "SourceConfig" member to the object itself.  New loads will do that.
 	FConfigFile NewEngineSettings;
 	FConfigCacheIni::LoadLocalIniFile(NewEngineSettings, TEXT("Engine"), true, *IniPlatformName(), true);
-	//the load above does not move settings from "SourceConfig" member to the object itself.  New loads will do that.
 	LuminEngineSettings = NewEngineSettings;
+	//remove the source config as it is just a copied pointer
+	LuminEngineSettings.SourceConfigFile = nullptr;
 	//override the android version too
 	EngineSettings = NewEngineSettings;
+	//remove the source config as it is just a copied pointer
+	EngineSettings.SourceConfigFile = nullptr;
 
 	// Get the Target RHIs for this platform, we do not always want all those that are supported.
 	TArray<FName> TargetedShaderFormats;
diff --git a/Engine/Source/Developer/Lumin/LuminTargetPlatform/Private/LuminTargetPlatform.h b/Engine/Source/Developer/Lumin/LuminTargetPlatform/Private/LuminTargetPlatform.h
--- a/Engine/Source/Developer/Lumin/LuminTargetPlatform/Private/LuminTargetPlatform.h
+++ b/Engine/Source/Developer/Lumin/LuminTargetPlatform/Private/LuminTargetPlatform.h
@@ -51,7 +51,7 @@ public:
 
 	virtual FString PlatformName() const override
 	{
-		return TEXT("Lumin");
+		return IsClientOnly() ? TEXT("LuminClient") : TEXT("Lumin");
 	}
 
 	virtual FString IniPlatformName() const override
diff --git a/Engine/Source/Developer/Lumin/LuminTargetPlatform/Private/LuminTargetPlatformModule.cpp b/Engine/Source/Developer/Lumin/LuminTargetPlatform/Private/LuminTargetPlatformModule.cpp
--- a/Engine/Source/Developer/Lumin/LuminTargetPlatform/Private/LuminTargetPlatformModule.cpp
+++ b/Engine/Source/Developer/Lumin/LuminTargetPlatform/Private/LuminTargetPlatformModule.cpp
@@ -6,11 +6,6 @@
 #include "Common/TargetPlatformBase.h"
 #include "LuminTargetPlatform.h"
 
-/**
- * Holds the target platform singleton.
- */
-static ITargetPlatform* LuminTargetSingleton = NULL;
-
 /**
  * Module for the Android target platform.
  */
@@ -21,9 +16,13 @@ public:
 	/**
 	 * Destructor.
 	 */
-	~FLuminTargetPlatformModule( )
+	virtual ~FLuminTargetPlatformModule( )
 	{
-		LuminTargetSingleton = NULL;
+		for (ITargetPlatform* TP : TargetPlatforms)
+		{
+			delete TP;
+		}
+		TargetPlatforms.Empty();
 	}
 
 
@@ -31,15 +30,15 @@ public:
 	
 	// Begin ITargetPlatformModule interface
 
-	virtual ITargetPlatform* GetTargetPlatform() override
+	virtual TArray<ITargetPlatform*> GetTargetPlatforms() override
 	{
-		if (LuminTargetSingleton == NULL)
+		if (TargetPlatforms.Num() == 0)
 		{
-			bool bClient = true;
-			LuminTargetSingleton = new FLuminTargetPlatform(bClient);
+			TargetPlatforms.Add(new FLuminTargetPlatform(false));
+			TargetPlatforms.Add(new FLuminTargetPlatform(true));
 		}
 		
-		return LuminTargetSingleton;
+		return TargetPlatforms;
 	}
 
 	// End ITargetPlatformModule interface
@@ -54,6 +53,10 @@ public:
 	{
 	}
 	// End IModuleInterface interface
+
+private:
+	/** Holds the target platforms. */
+	TArray<ITargetPlatform*> TargetPlatforms;
 };
 
 
diff --git a/Engine/Source/Programs/AutomationTool/Lumin/LuminPlatform.Automation.cs b/Engine/Source/Programs/AutomationTool/Lumin/LuminPlatform.Automation.cs
--- a/Engine/Source/Programs/AutomationTool/Lumin/LuminPlatform.Automation.cs
+++ b/Engine/Source/Programs/AutomationTool/Lumin/LuminPlatform.Automation.cs
@@ -244,7 +244,7 @@ public class LuminPlatform : Platform
 	/// <returns>Cook platform string.</returns>
 	public override string GetCookPlatform(bool bDedicatedServer, bool bIsClientOnly)
 	{
-		return "Lumin";
+		return bIsClientOnly ? "LuminClient" : "Lumin";
 	}
 
 	/// <summary>
