PR #5200: Log file path issue (Contributed by [REDACTED_PERSON])

diff --git a/Engine/Source/Developer/DesktopPlatform/Private/DesktopPlatformBase.cpp b/Engine/Source/Developer/DesktopPlatform/Private/DesktopPlatformBase.cpp
--- a/Engine/Source/Developer/DesktopPlatform/Private/DesktopPlatformBase.cpp
+++ b/Engine/Source/Developer/DesktopPlatform/Private/DesktopPlatformBase.cpp
@@ -610,7 +610,7 @@ bool FDesktopPlatformBase::GenerateProjectFiles(const FString& RootDir, const FS
 
 	if (!LogFilePath.IsEmpty())
 	{
-		Arguments += TEXT(" -log=") + LogFilePath;
+		Arguments += FString::Printf(TEXT(" -log=\"%s\""), *LogFilePath);
 	}
 
 	// Compile UnrealBuildTool if it doesn't exist. This can happen if we're just copying source from somewhere.
