Accessible Platform list no longer includes any platform which has no Configuration directory present.

diff --git a/Engine/Source/Developer/DesktopPlatform/Private/PlatformInfo.cpp b/Engine/Source/Developer/DesktopPlatform/Private/PlatformInfo.cpp
--- a/Engine/Source/Developer/DesktopPlatform/Private/PlatformInfo.cpp
+++ b/Engine/Source/Developer/DesktopPlatform/Private/PlatformInfo.cpp
@@ -3,6 +3,8 @@
 #include "PlatformInfo.h"
 #include "DesktopPlatformPrivate.h"
 #include "Misc/DataDrivenPlatformInfoRegistry.h"
+#include "HAL/FileManager.h"
+#include "Misc/Paths.h"
 
 #define LOCTEXT_NAMESPACE "PlatformInfo"
 
@@ -209,13 +211,18 @@ FPlatformEnumerator EnumeratePlatformInfoArray(bool bAccessiblePlatformsOnly)
 
 			for (const FPlatformInfo& PlatformInfo : Enumerator)
 			{
-				if (PlatformInfo.bIsConfidential && ConfidentalPlatforms.Contains(PlatformInfo.IniPlatformName))
+				// We check for the existence of the configuration directory to include the platform
+				// P4/Git may have filtered out platforms and we don't want to include the platform to keep code from trying to load other files that don't exist.
+				if (IFileManager::Get().DirectoryExists(*(FPaths::RootDir() / TEXT("Engine") / TEXT("Config") / PlatformInfo.IniPlatformName)))
 				{
-					AccessiblePlatforms.Add(PlatformInfo);
-				}
-				else if(!PlatformInfo.bIsConfidential)
-				{
-					AccessiblePlatforms.Add(PlatformInfo);
+					if (PlatformInfo.bIsConfidential && ConfidentalPlatforms.Contains(PlatformInfo.IniPlatformName))
+					{
+						AccessiblePlatforms.Add(PlatformInfo);
+					}
+					else if (!PlatformInfo.bIsConfidential)
+					{
+						AccessiblePlatforms.Add(PlatformInfo);
+					}
 				}
 			}
 
