Accessible Platform list no longer includes any platform which has no Configuration directory with files in it.

diff --git a/Engine/Source/Developer/DesktopPlatform/Private/PlatformInfo.cpp b/Engine/Source/Developer/DesktopPlatform/Private/PlatformInfo.cpp
--- a/Engine/Source/Developer/DesktopPlatform/Private/PlatformInfo.cpp
+++ b/Engine/Source/Developer/DesktopPlatform/Private/PlatformInfo.cpp
@@ -211,9 +211,13 @@ FPlatformEnumerator EnumeratePlatformInfoArray(bool bAccessiblePlatformsOnly)
 
 			for (const FPlatformInfo& PlatformInfo : Enumerator)
 			{
-				// We check for the existence of the configuration directory to include the platform
-				// P4/Git may have filtered out platforms and we don't want to include the platform to keep code from trying to load other files that don't exist.
-				if (IFileManager::Get().DirectoryExists(*(FPaths::RootDir() / TEXT("Engine") / TEXT("Config") / PlatformInfo.IniPlatformName)))
+				const FString IniFolderPath = FPaths::RootDir() / TEXT("Engine") / TEXT("Config") / PlatformInfo.IniPlatformName;
+				TArray<FString> FoundFiles;
+				IFileManager::Get().FindFiles(FoundFiles, *IniFolderPath);
+			
+				// We check that the configuration directory exists with actual files in it to include the platform
+				// P4/Git may have filtered out platforms and we don't want to include filtered platform to keep code from trying to load other files that don't exist.
+				if (FoundFiles.Num() > 0)
 				{
 					if (PlatformInfo.bIsConfidential && ConfidentalPlatforms.Contains(PlatformInfo.IniPlatformName))
 					{
