Fixed display of supported platforms, platforms that are missing or stripped shouldn't show up any longer.

diff --git a/Engine/Source/Developer/DesktopPlatform/Private/PlatformInfo.cpp b/Engine/Source/Developer/DesktopPlatform/Private/PlatformInfo.cpp
--- a/Engine/Source/Developer/DesktopPlatform/Private/PlatformInfo.cpp
+++ b/Engine/Source/Developer/DesktopPlatform/Private/PlatformInfo.cpp
@@ -249,13 +249,13 @@ FPlatformEnumerator EnumeratePlatformInfoArray(bool bAccessiblePlatformsOnly)
 	}
 }
 
-TArray<FVanillaPlatformEntry> BuildPlatformHierarchy(const EPlatformFilter InFilter)
+TArray<FVanillaPlatformEntry> BuildPlatformHierarchy(const EPlatformFilter InFilter, bool bAccessiblePlatformsOnly)
 {
 	TArray<FVanillaPlatformEntry> VanillaPlatforms;
 
 	// Build up a tree from the platforms we support (vanilla outers, with a list of flavors)
 	// PlatformInfoArray should be ordered in such a way that the vanilla platforms always appear before their flavors
-	for(const FPlatformInfo& PlatformInfo : AllPlatformInfoArray)
+	for (const PlatformInfo::FPlatformInfo& PlatformInfo : EnumeratePlatformInfoArray(bAccessiblePlatformsOnly))
 	{
 		if(PlatformInfo.IsVanilla())
 		{
@@ -287,7 +287,7 @@ TArray<FVanillaPlatformEntry> BuildPlatformHierarchy(const EPlatformFilter InFil
 	return VanillaPlatforms;
 }
 
-FVanillaPlatformEntry BuildPlatformHierarchy(const FName& InPlatformName, const EPlatformFilter InFilter)
+FVanillaPlatformEntry BuildPlatformHierarchy(const FName& InPlatformName, const EPlatformFilter InFilter, bool bAccessiblePlatformsOnly)
 {
 	FVanillaPlatformEntry VanillaPlatformEntry;
 	const FPlatformInfo* VanillaPlatformInfo = FindVanillaPlatformInfo(InPlatformName);
@@ -296,7 +296,7 @@ FVanillaPlatformEntry BuildPlatformHierarchy(const FName& InPlatformName, const
 	{
 		VanillaPlatformEntry.PlatformInfo = VanillaPlatformInfo;
 		
-		for (const FPlatformInfo& PlatformInfo : AllPlatformInfoArray)
+		for (const PlatformInfo::FPlatformInfo& PlatformInfo : EnumeratePlatformInfoArray(bAccessiblePlatformsOnly))
 		{
 			if (!PlatformInfo.IsVanilla() && PlatformInfo.VanillaPlatformName == VanillaPlatformInfo->PlatformInfoName)
 			{
diff --git a/Engine/Source/Developer/DesktopPlatform/Public/PlatformInfo.h b/Engine/Source/Developer/DesktopPlatform/Public/PlatformInfo.h
--- a/Engine/Source/Developer/DesktopPlatform/Public/PlatformInfo.h
+++ b/Engine/Source/Developer/DesktopPlatform/Public/PlatformInfo.h
@@ -333,17 +333,19 @@ namespace PlatformInfo
 	/**
 	 * Build a hierarchy mapping vanilla platforms to their flavors
 	 * @param InFilter - Flags to control which kinds of flavors you want to include
+	 * @param bAccessiblePlatformsOnly	If true, only the accessible platforms(installed, or could be installed) will be returned
 	 * @return An array of vanilla platforms, potentially containing flavors
 	 */
-	DESKTOPPLATFORM_API TArray<FVanillaPlatformEntry> BuildPlatformHierarchy(const EPlatformFilter InFilter);
+	DESKTOPPLATFORM_API TArray<FVanillaPlatformEntry> BuildPlatformHierarchy(const EPlatformFilter InFilter, bool bAccessiblePlatformsOnly = true);
 
 	/**
 	* Build a hierarchy mapping for specified vanilla platform to it flavors
 	* @param InPlatformName - Platform name to build hierarchy for, could be vanilla or flavor name
 	* @param InFilter - Flags to control which kinds of flavors you want to include
+	* @param bAccessiblePlatformsOnly	If true, only the accessible platforms(installed, or could be installed) will be returned
 	* @return Vanilla platform potentially containing flavors
 	*/
-	DESKTOPPLATFORM_API FVanillaPlatformEntry BuildPlatformHierarchy(const FName& InPlatformName, const EPlatformFilter InFilter);
+	DESKTOPPLATFORM_API FVanillaPlatformEntry BuildPlatformHierarchy(const FName& InPlatformName, const EPlatformFilter InFilter, bool bAccessiblePlatformsOnly = true);
 
 	/**
 	 * Get an array of all the platforms we know about
