Suppress warning message when opening the project browser if a project doesn't have any engine identifier or version field.

diff --git a/Engine/Source/Developer/DesktopPlatform/Private/DesktopPlatformBase.cpp b/Engine/Source/Developer/DesktopPlatform/Private/DesktopPlatformBase.cpp
--- a/Engine/Source/Developer/DesktopPlatform/Private/DesktopPlatformBase.cpp
+++ b/Engine/Source/Developer/DesktopPlatform/Private/DesktopPlatformBase.cpp
@@ -491,8 +491,8 @@ bool FDesktopPlatformBase::GetEngineIdentifierForProject(const FString& ProjectF
 	}
 
 	// Otherwise check the engine version string for 4.0, in case this project existed before the engine association stuff went in
-	FString EngineVersionString = ProjectFile->GetStringField(TEXT("EngineVersion"));
-	if(EngineVersionString.Len() > 0)
+	FString EngineVersionString;
+	if(ProjectFile->TryGetStringField(TEXT("EngineVersion"), EngineVersionString) && EngineVersionString.Len() > 0)
 	{
 		FEngineVersion EngineVersion;
 		if(FEngineVersion::Parse(EngineVersionString, EngineVersion) && EngineVersion.HasChangelist() && EngineVersion.ToString(EVersionComponent::Minor) == TEXT("4.0"))
