Couple more CIS fixes for target platforms.

#rnx

diff --git a/Engine/Source/Developer/DesktopPlatform/Private/Linux/DesktopPlatformLinux.cpp b/Engine/Source/Developer/DesktopPlatform/Private/Linux/DesktopPlatformLinux.cpp
--- a/Engine/Source/Developer/DesktopPlatform/Private/Linux/DesktopPlatformLinux.cpp
+++ b/Engine/Source/Developer/DesktopPlatform/Private/Linux/DesktopPlatformLinux.cpp
@@ -251,7 +251,7 @@ static bool CompareAndCheckDesktopFile(const TCHAR* DesktopFileName, const TCHAR
 
 	// Get the contents of the desktop file.
 	FString InstalledDesktopFileContents;
-	FFileHelper::LoadFileToString(InstalledDesktopFileContents, *FString::Printf(TEXT("%s/applications/%s.desktop"), DataDir, *Association));
+	FFileHelper::LoadFileToString(InstalledDesktopFileContents, *FString::Printf(TEXT("%s/applications/%s.desktop"), *DataDir, *Association));
 
 	// Make sure the installed and default desktop file was created by unreal engine.
 	if (!InstalledDesktopFileContents.Contains(TEXT("Comment=Created by Unreal Engine")))
