Reset module paths after project compilation in case they have changed.

diff --git a/Engine/Source/Developer/DesktopPlatform/Private/DesktopPlatformBase.cpp b/Engine/Source/Developer/DesktopPlatform/Private/DesktopPlatformBase.cpp
--- a/Engine/Source/Developer/DesktopPlatform/Private/DesktopPlatformBase.cpp
+++ b/Engine/Source/Developer/DesktopPlatform/Private/DesktopPlatformBase.cpp
@@ -558,11 +558,13 @@ bool FDesktopPlatformBase::CleanGameProject(const FString& ProjectDir, FString&
 
 bool FDesktopPlatformBase::CompileGameProject(const FString& RootDir, const FString& ProjectFileName, FFeedbackContext* Warn)
 {
+	FModuleManager& ModuleManager = FModuleManager::Get();
+
 	// Get the project directory
 	FString ProjectDir = FPaths::GetPath(ProjectFileName);
 
 	// Build the argument list
-	FString Arguments = FString::Printf(TEXT("%s %s"), FModuleManager::Get().GetUBTConfiguration(), FPlatformMisc::GetUBTPlatform());
+	FString Arguments = FString::Printf(TEXT("%s %s"), ModuleManager.GetUBTConfiguration(), FPlatformMisc::GetUBTPlatform());
 
 	// Append the project name if it's a foreign project
 	if ( !ProjectFileName.IsEmpty() )
@@ -574,7 +576,12 @@ bool FDesktopPlatformBase::CompileGameProject(const FString& RootDir, const FStr
 	Arguments += " -TargetType=Editor -Progress -NoHotReloadFromIDE";
 
 	// Run UBT
-	return RunUnrealBuildTool(LOCTEXT("CompilingProject", "Compiling project..."), RootDir, Arguments, Warn);
+	bool bResult = RunUnrealBuildTool(LOCTEXT("CompilingProject", "Compiling project..."), RootDir, Arguments, Warn);
+
+	// Reset module paths in case they have changed during compilation
+	ModuleManager.ResetModulePathsCache();
+
+	return bResult;
 }
 
 bool FDesktopPlatformBase::GenerateProjectFiles(const FString& RootDir, const FString& ProjectFileName, FFeedbackContext* Warn, FString LogFilePath)
