Fix UnrealVersionSelectors slate window and defaults for running a uproject shortcut

diff --git a/Engine/Source/Developer/DesktopPlatform/Private/Linux/DesktopPlatformLinux.cpp b/Engine/Source/Developer/DesktopPlatform/Private/Linux/DesktopPlatformLinux.cpp
--- a/Engine/Source/Developer/DesktopPlatform/Private/Linux/DesktopPlatformLinux.cpp
+++ b/Engine/Source/Developer/DesktopPlatform/Private/Linux/DesktopPlatformLinux.cpp
@@ -336,7 +336,7 @@ bool FDesktopPlatformLinux::UpdateFileAssociations()
 	{
 		return false;
 	}
-	
+
 	if (!RunXDGUtil(FString::Printf(TEXT("xdg-icon-resource install --novendor --mode user --context apps --size 256 %sRuntime/Launch/Resources/Linux/UE4.png ubinary"), *FPaths::EngineSourceDir())))
 	{
 		return false;
@@ -354,7 +354,7 @@ bool FDesktopPlatformLinux::UpdateFileAssociations()
 		return false;
 	}
 
-	// Add the desktop file for the Unreal Engine icon from the template.
+	// Add the desktop file for the Unreal Engine Generate Project List icon from the template.
 	DesktopTemplate = FString();
 	FFileHelper::LoadFileToString(DesktopTemplate, *FString::Printf(TEXT("%sPrograms/UnrealVersionSelector/Private/Linux/Resources/com.epicgames.UnrealEngine.desktop"), *FPaths::EngineSourceDir()));      
 	DesktopTemplate = DesktopTemplate.Replace(TEXT("*ENGINEDIR*"), *AbsoluteEngineDir);
@@ -364,12 +364,22 @@ bool FDesktopPlatformLinux::UpdateFileAssociations()
 		return false;
 	}
 
+	// Add the desktop file for the Unreal Engine Editor icon from the template.
+	DesktopTemplate = FString();
+	FFileHelper::LoadFileToString(DesktopTemplate, *FString::Printf(TEXT("%sPrograms/UnrealVersionSelector/Private/Linux/Resources/com.epicgames.UnrealEngineEditor.desktop"), *FPaths::EngineSourceDir()));      
+	DesktopTemplate = DesktopTemplate.Replace(TEXT("*ENGINEDIR*"), *AbsoluteEngineDir);
+	FFileHelper::SaveStringToFile(DesktopTemplate, TEXT("/tmp/com.epicgames.UnrealEngineEditor.desktop"));
+	if (!RunXDGUtil(TEXT("xdg-desktop-menu install --novendor --mode user /tmp/com.epicgames.UnrealEngineEditor.desktop")))
+	{
+		return false;
+	}
+
 	// Create the mime types and set the default applications.
 	if (!RunXDGUtil(FString::Printf(TEXT("xdg-mime install --novendor --mode user %sPrograms/UnrealVersionSelector/Private/Linux/Resources/uproject.xml"), *FPaths::EngineSourceDir())))
 	{
 		return false;
 	}
-	if (!RunXDGUtil(TEXT("xdg-mime default com.epicgames.UnrealVersionSelector.desktop application/uproject")))
+	if (!RunXDGUtil(TEXT("xdg-mime default com.epicgames.UnrealEngineEditor.desktop application/uproject")))
 	{
 		return false;
 	}
diff --git a/Engine/Source/Programs/UnrealVersionSelector/Private/Linux/LinuxPlatformInstallation.cpp b/Engine/Source/Programs/UnrealVersionSelector/Private/Linux/LinuxPlatformInstallation.cpp
--- a/Engine/Source/Programs/UnrealVersionSelector/Private/Linux/LinuxPlatformInstallation.cpp
+++ b/Engine/Source/Programs/UnrealVersionSelector/Private/Linux/LinuxPlatformInstallation.cpp
@@ -150,8 +150,7 @@ public:
 						SNew(SHorizontalBox)
 
 						+SHorizontalBox::Slot()
-							.AutoWidth()
-							.MaxWidth(450.0f)
+							.FillWidth(1)
 							.HAlign(HAlign_Fill)
 							.VAlign(VAlign_Center)
 							.Padding(2.0f)
@@ -163,8 +162,7 @@ public:
 									.InitiallySelectedItem(SelectBuildInfo->SelectedEngineInstallationInfo)
 									[
 										SNew(STextBlock)
-											.MinDesiredWidth(425.0f)
-											.Text(FText::FromString(SelectBuildInfo->SelectedEngineInstallationInfo->Description))
+											.Text(this, &SSelectBuildDialog::GetSelectedEngineInstallDescription)
 									]
 							]
 
@@ -227,6 +225,11 @@ private:
 					.Text(FText::FromString(Item->Description));
 	}
 
+	FText GetSelectedEngineInstallDescription() const
+	{
+		return FText::FromString(SelectBuildInfo->SelectedEngineInstallationInfo->Description);
+	}
+
 	void OnSelectionChanged(TSharedPtr<FEngineInstallationInfo> Item, ESelectInfo::Type SelectInfo)
 	{
 		if (Item.IsValid())
@@ -532,4 +535,4 @@ void FLinuxPlatformInstallation::ErrorDialog(const FString &Message, const FStri
 {
 	FErrorDialog Dialog(Message, LogText);
 	Dialog.DoModal();
-}
\ No newline at end of file
+}
diff --git a/Engine/Source/Programs/UnrealVersionSelector/Private/UnrealVersionSelector.cpp b/Engine/Source/Programs/UnrealVersionSelector/Private/UnrealVersionSelector.cpp
--- a/Engine/Source/Programs/UnrealVersionSelector/Private/UnrealVersionSelector.cpp
+++ b/Engine/Source/Programs/UnrealVersionSelector/Private/UnrealVersionSelector.cpp
@@ -281,7 +281,7 @@ int Main(const TArray<FString>& Arguments)
 	else if (Arguments.Num() == 2 && Arguments[0] == TEXT("-editor"))
 	{
 		// Open a project with the editor
-		bRes = LaunchEditor(Arguments[1], L"");
+		bRes = LaunchEditor(Arguments[1], "");
 	}
 	else if (Arguments[0] == TEXT("-projectlist"))
 	{
@@ -291,7 +291,7 @@ int Main(const TArray<FString>& Arguments)
 	else if (Arguments.Num() == 2 && Arguments[0] == TEXT("-game"))
 	{
 		// Play a game using the editor executable
-		bRes = LaunchEditor(Arguments[1], L"-game");
+		bRes = LaunchEditor(Arguments[1], "-game");
 	}
 	else if (Arguments.Num() == 2 && Arguments[0] == TEXT("-projectfiles"))
 	{
