Stubbed DesktopPlatform support on unknown platforms so that PlatformInfo can still be used

diff --git a/Engine/Source/Developer/DesktopPlatform/Private/DesktopPlatformPrivate.h b/Engine/Source/Developer/DesktopPlatform/Private/DesktopPlatformPrivate.h
--- a/Engine/Source/Developer/DesktopPlatform/Private/DesktopPlatformPrivate.h
+++ b/Engine/Source/Developer/DesktopPlatform/Private/DesktopPlatformPrivate.h
@@ -10,6 +10,8 @@
 #include "Mac/DesktopPlatformMac.h"
 #elif PLATFORM_LINUX
 #include "Linux/DesktopPlatformLinux.h"
+#else
+#include "DesktopPlatformStub.h"
 #endif
 
 DECLARE_LOG_CATEGORY_EXTERN(LogDesktopPlatform, Log, All);
diff --git a/Engine/Source/Developer/DesktopPlatform/Private/DesktopPlatformStub.h b/Engine/Source/Developer/DesktopPlatform/Private/DesktopPlatformStub.h
new file mode 100644
--- /dev/null
+++ b/Engine/Source/Developer/DesktopPlatform/Private/DesktopPlatformStub.h
@@ -0,0 +1,33 @@
+// Copyright 1998-2019 Epic Games, Inc. All Rights Reserved.
+
+#pragma once
+
+#include "CoreMinimal.h"
+#include "DesktopPlatformBase.h"
+#include "Misc/Paths.h"
+
+class FDesktopPlatformStub : public FDesktopPlatformBase
+{
+public:
+	// IDesktopPlatform Implementation
+	virtual bool OpenFileDialog(const void* ParentWindowHandle, const FString& DialogTitle, const FString& DefaultPath, const FString& DefaultFile, const FString& FileTypes, uint32 Flags, TArray<FString>& OutFilenames) override { return false; }
+	virtual bool OpenFileDialog(const void* ParentWindowHandle, const FString& DialogTitle, const FString& DefaultPath, const FString& DefaultFile, const FString& FileTypes, uint32 Flags, TArray<FString>& OutFilenames, int32& OutFilterIndex) override { return false; }
+	virtual bool SaveFileDialog(const void* ParentWindowHandle, const FString& DialogTitle, const FString& DefaultPath, const FString& DefaultFile, const FString& FileTypes, uint32 Flags, TArray<FString>& OutFilenames) override { return false; }
+	virtual bool OpenDirectoryDialog(const void* ParentWindowHandle, const FString& DialogTitle, const FString& DefaultPath, FString& OutFolderName) override { return false; }
+	virtual bool OpenFontDialog(const void* ParentWindowHandle, FString& OutFontName, float& OutHeight, EFontImportFlags& OutFlags) override { return false; }
+
+	virtual bool RegisterEngineInstallation(const FString &RootDir, FString &OutIdentifier) override { return false; }
+	virtual void EnumerateEngineInstallations(TMap<FString, FString> &OutInstallations) override {}
+
+	virtual bool VerifyFileAssociations() override { return false; }
+	virtual bool UpdateFileAssociations() override { return false; }
+
+	virtual bool RunUnrealBuildTool(const FText& Description, const FString& RootDir, const FString& Arguments, FFeedbackContext* Warn) override { return false; }
+	virtual bool IsUnrealBuildToolRunning() override { return false; }
+
+	virtual FFeedbackContext* GetNativeFeedbackContext() override { return nullptr; }
+
+	virtual FString GetUserTempPath() override { return FPaths::ProjectIntermediateDir() + TEXT("UserTemp/"); }
+};
+
+typedef FDesktopPlatformStub FDesktopPlatform;
diff --git a/Engine/Source/Developer/DesktopPlatform/Private/PlatformInfo.cpp b/Engine/Source/Developer/DesktopPlatform/Private/PlatformInfo.cpp
--- a/Engine/Source/Developer/DesktopPlatform/Private/PlatformInfo.cpp
+++ b/Engine/Source/Developer/DesktopPlatform/Private/PlatformInfo.cpp
@@ -91,6 +91,11 @@ static const bool IsAvailableOnWindows = false;
 static const bool IsAvailableOnMac = false;
 static const bool IsAvailableOnLinux = true;
 static const FString IOSTutorial = TEXT("/Engine/Tutorial/NotYetImplemented");
+#else
+static const bool IsAvailableOnWindows = false;
+static const bool IsAvailableOnMac = false;
+static const bool IsAvailableOnLinux = false;
+static const FString IOSTutorial = TEXT("/Engine/Tutorial/NotYetImplemented");
 #endif
 
 static const FPlatformInfo AllPlatformInfoArray[] = {
