Dont add to a container while using iterators. Invalidates the iterators

diff --git a/Engine/Source/Developer/DesktopPlatform/Private/Linux/DesktopPlatformLinux.cpp b/Engine/Source/Developer/DesktopPlatform/Private/Linux/DesktopPlatformLinux.cpp
--- a/Engine/Source/Developer/DesktopPlatform/Private/Linux/DesktopPlatformLinux.cpp
+++ b/Engine/Source/Developer/DesktopPlatform/Private/Linux/DesktopPlatformLinux.cpp
@@ -169,6 +169,8 @@ void FDesktopPlatformLinux::EnumerateEngineInstallations(TMap<FString, FString>
 		Section.Remove(Key);
 	}
 
+	FConfigSection SectionsToAdd;
+
 	// Iterate through all entries.
 	for (auto It : Section)
 	{
@@ -178,6 +180,11 @@ void FDesktopPlatformLinux::EnumerateEngineInstallations(TMap<FString, FString>
 
 		FString EngineId;
 		const FName* Key = Section.FindKey(EngineDir);
+		if (Key == nullptr)
+		{
+			Key = SectionsToAdd.FindKey(EngineDir);
+		}
+
 		if (Key)
 		{
 			FGuid IdGuid;
@@ -189,7 +196,8 @@ void FDesktopPlatformLinux::EnumerateEngineInstallations(TMap<FString, FString>
 			if (!OutInstallations.FindKey(EngineDir))
 			{
 				EngineId = FGuid::NewGuid().ToString(EGuidFormats::DigitsWithHyphensInBraces);
-				Section.AddUnique(*EngineId, EngineDir);
+				SectionsToAdd.AddUnique(*EngineId, EngineDir);
+
 				ConfigFile.Dirty = true;
 			}
 		}
@@ -199,6 +207,11 @@ void FDesktopPlatformLinux::EnumerateEngineInstallations(TMap<FString, FString>
 		}
 	}
 
+	for (auto It : SectionsToAdd)
+	{
+		Section.AddUnique(It.Key, It.Value.GetValue());
+	}
+
 	ConfigFile.Write(ConfigPath);
 
 	IFileManager::Get().Delete(*UProjectPath);
