Rename EBuildTargets::Type to EBuildTargetType, and add missing entries for server and program targets. Also replace usages of PlatformInfo::EPlatformType with EBuildTargetType.

#rnx

diff --git a/Engine/Platforms/XXX/Source/Developer/XXXTargetPlatform/Private/XXXPlatformProperties.h b/Engine/Platforms/XXX/Source/Developer/XXXTargetPlatform/Private/XXXPlatformProperties.h
--- a/Engine/Platforms/XXX/Source/Developer/XXXTargetPlatform/Private/XXXPlatformProperties.h
+++ b/Engine/Platforms/XXX/Source/Developer/XXXTargetPlatform/Private/XXXPlatformProperties.h
@@ -62,9 +62,9 @@ struct FXXXPlatformProperties : public FGenericPlatformProperties
 		return false;
 	}
 
-	static FORCEINLINE bool SupportsBuildTarget( EBuildTargets::Type BuildTarget )
+	static FORCEINLINE bool SupportsBuildTarget( EBuildTargetType TargetType )
 	{
-		return (BuildTarget == EBuildTargets::Game);
+		return (TargetType == EBuildTargetType::Game);
 	}
 
 	static FORCEINLINE bool SupportsAutoSDK()
diff --git a/Engine/Platforms/XXX/Source/Developer/XXXTargetPlatform/Private/XXXTargetDevice.cpp b/Engine/Platforms/XXX/Source/Developer/XXXTargetPlatform/Private/XXXTargetDevice.cpp
--- a/Engine/Platforms/XXX/Source/Developer/XXXTargetPlatform/Private/XXXTargetDevice.cpp
+++ b/Engine/Platforms/XXX/Source/Developer/XXXTargetPlatform/Private/XXXTargetDevice.cpp
@@ -46,7 +46,7 @@ int32 FXXXTargetDevice::GetProcessSnapshot( TArray<FTargetDeviceProcessInfo>& Ou
 }
 
 
-bool FXXXTargetDevice::Launch( const FString& AppId, EBuildConfigurations::Type BuildConfiguration, EBuildTargets::Type BuildTarget, const FString& Params, uint32* OutProcessId )
+bool FXXXTargetDevice::Launch( const FString& AppId, EBuildConfigurations::Type BuildConfiguration, EBuildTargetType TargetType, const FString& Params, uint32* OutProcessId )
 {
 	checkNoEntry();
 	return false;
diff --git a/Engine/Platforms/XXX/Source/Developer/XXXTargetPlatform/Private/XXXTargetDevice.h b/Engine/Platforms/XXX/Source/Developer/XXXTargetPlatform/Private/XXXTargetDevice.h
--- a/Engine/Platforms/XXX/Source/Developer/XXXTargetPlatform/Private/XXXTargetDevice.h
+++ b/Engine/Platforms/XXX/Source/Developer/XXXTargetPlatform/Private/XXXTargetDevice.h
@@ -97,7 +97,7 @@ public:
 		return CachedDefault;
 	}
 
-	virtual bool Launch( const FString& AppId, EBuildConfigurations::Type BuildConfiguration, EBuildTargets::Type BuildTarget, const FString& Params, uint32* OutProcessId ) override;
+	virtual bool Launch( const FString& AppId, EBuildConfigurations::Type BuildConfiguration, EBuildTargetType TargetType, const FString& Params, uint32* OutProcessId ) override;
 
 	virtual bool PowerOff( bool Force ) override;
 
diff --git a/Engine/Plugins/Media/MediaPlayerEditor/Source/MediaPlayerEditor/Private/Customizations/BaseMediaSourceCustomization.cpp b/Engine/Plugins/Media/MediaPlayerEditor/Source/MediaPlayerEditor/Private/Customizations/BaseMediaSourceCustomization.cpp
--- a/Engine/Plugins/Media/MediaPlayerEditor/Source/MediaPlayerEditor/Private/Customizations/BaseMediaSourceCustomization.cpp
+++ b/Engine/Plugins/Media/MediaPlayerEditor/Source/MediaPlayerEditor/Private/Customizations/BaseMediaSourceCustomization.cpp
@@ -137,7 +137,7 @@ TSharedRef<SWidget> FBaseMediaSourceCustomization::MakePlatformPlayerNamesValueW
 
 	for (const PlatformInfo::FPlatformInfo& PlatformInfo : PlatformInfo::GetPlatformInfoArray())
 	{
-		if (PlatformInfo.IsVanilla() && (PlatformInfo.PlatformType == PlatformInfo::EPlatformType::Game) && (PlatformInfo.PlatformInfoName != TEXT("AllDesktop")))
+		if (PlatformInfo.IsVanilla() && (PlatformInfo.PlatformType == EBuildTargetType::Game) && (PlatformInfo.PlatformInfoName != TEXT("AllDesktop")))
 		{
 			AvailablePlatforms.Add(&PlatformInfo);
 		}
diff --git a/Engine/Plugins/Media/MediaPlayerEditor/Source/MediaPlayerEditor/Private/Customizations/PlatformMediaSourceCustomization.cpp b/Engine/Plugins/Media/MediaPlayerEditor/Source/MediaPlayerEditor/Private/Customizations/PlatformMediaSourceCustomization.cpp
--- a/Engine/Plugins/Media/MediaPlayerEditor/Source/MediaPlayerEditor/Private/Customizations/PlatformMediaSourceCustomization.cpp
+++ b/Engine/Plugins/Media/MediaPlayerEditor/Source/MediaPlayerEditor/Private/Customizations/PlatformMediaSourceCustomization.cpp
@@ -75,7 +75,7 @@ TSharedRef<SWidget> FPlatformMediaSourceCustomization::MakePlatformMediaSourcesV
 
 	for (const PlatformInfo::FPlatformInfo& PlatformInfo : PlatformInfo::GetPlatformInfoArray())
 	{
-		if (PlatformInfo.IsVanilla() && (PlatformInfo.PlatformType == PlatformInfo::EPlatformType::Game) && (PlatformInfo.PlatformInfoName != TEXT("AllDesktop")))
+		if (PlatformInfo.IsVanilla() && (PlatformInfo.PlatformType == EBuildTargetType::Game) && (PlatformInfo.PlatformInfoName != TEXT("AllDesktop")))
 		{
 			if (PlatformInfo.PlatformInfoName == TEXT("TVOS"))
 			{
diff --git a/Engine/Plugins/NetcodeUnitTest/NetcodeUnitTest/Source/NetcodeUnitTest/Classes/ProcessUnitTest.h b/Engine/Plugins/NetcodeUnitTest/NetcodeUnitTest/Source/NetcodeUnitTest/Classes/ProcessUnitTest.h
--- a/Engine/Plugins/NetcodeUnitTest/NetcodeUnitTest/Source/NetcodeUnitTest/Classes/ProcessUnitTest.h
+++ b/Engine/Plugins/NetcodeUnitTest/NetcodeUnitTest/Source/NetcodeUnitTest/Classes/ProcessUnitTest.h
@@ -243,7 +243,7 @@ protected:
 	 * @return					Returns a pointer to the new processes handling struct
 	 */
 	virtual TWeakPtr<FUnitTestProcess> StartUE4UnitTestProcess(FString InCommandline, bool bMinimized=true,
-																EBuildTargets::Type Type=EBuildTargets::Game);
+																EBuildTargetType Type=EBuildTargetType::Game);
 
 	/**
 	 * Shuts-down/cleans-up a child process tied to the unit test
diff --git a/Engine/Plugins/NetcodeUnitTest/NetcodeUnitTest/Source/NetcodeUnitTest/Private/ClientUnitTest.cpp b/Engine/Plugins/NetcodeUnitTest/NetcodeUnitTest/Source/NetcodeUnitTest/Private/ClientUnitTest.cpp
--- a/Engine/Plugins/NetcodeUnitTest/NetcodeUnitTest/Source/NetcodeUnitTest/Private/ClientUnitTest.cpp
+++ b/Engine/Plugins/NetcodeUnitTest/NetcodeUnitTest/Source/NetcodeUnitTest/Private/ClientUnitTest.cpp
@@ -1165,7 +1165,7 @@ void UClientUnitTest::StartUnitTestServer()
 			ServerParameters += FString::Printf(TEXT(" -BeaconPort=%i -NUTMonitorBeacon"), ServerBeaconPort);
 		}
 
-		ServerHandle = StartUE4UnitTestProcess(ServerParameters, true, EBuildTargets::Server);
+		ServerHandle = StartUE4UnitTestProcess(ServerParameters, true, EBuildTargetType::Server);
 
 		if (ServerHandle.IsValid())
 		{
diff --git a/Engine/Plugins/NetcodeUnitTest/NetcodeUnitTest/Source/NetcodeUnitTest/Private/ProcessUnitTest.cpp b/Engine/Plugins/NetcodeUnitTest/NetcodeUnitTest/Source/NetcodeUnitTest/Private/ProcessUnitTest.cpp
--- a/Engine/Plugins/NetcodeUnitTest/NetcodeUnitTest/Source/NetcodeUnitTest/Private/ProcessUnitTest.cpp
+++ b/Engine/Plugins/NetcodeUnitTest/NetcodeUnitTest/Source/NetcodeUnitTest/Private/ProcessUnitTest.cpp
@@ -98,14 +98,14 @@ TWeakPtr<FUnitTestProcess> UProcessUnitTest::StartUnitTestProcess(FString Path,
 }
 
 TWeakPtr<FUnitTestProcess> UProcessUnitTest::StartUE4UnitTestProcess(FString InCommandline, bool bMinimized/*=true*/,
-																		EBuildTargets::Type Type/*=EBuildTargets::Game*/)
+																		EBuildTargetType Type/*=EBuildTargetType::Game*/)
 {
 	TWeakPtr<FUnitTestProcess> ReturnVal = nullptr;
 	FString TargetExecutable = FApp::GetName();
 
 #if UE_BUILD_DEVELOPMENT && !WITH_EDITOR
 	// Development modes other than Dev Editor, must target the separate Server process
-	if (Type == EBuildTargets::Server)
+	if (Type == EBuildTargetType::Server)
 	{
 		TargetExecutable = TargetExecutable.Replace(TEXT("Game"), TEXT("Server"));
 
diff --git a/Engine/Plugins/Runtime/AR/Microsoft/HoloLensAR/Source/HoloLensTargetPlatform/Private/HoloLensTargetDevice.cpp b/Engine/Plugins/Runtime/AR/Microsoft/HoloLensAR/Source/HoloLensTargetPlatform/Private/HoloLensTargetDevice.cpp
--- a/Engine/Plugins/Runtime/AR/Microsoft/HoloLensAR/Source/HoloLensTargetPlatform/Private/HoloLensTargetDevice.cpp
+++ b/Engine/Plugins/Runtime/AR/Microsoft/HoloLensAR/Source/HoloLensTargetPlatform/Private/HoloLensTargetDevice.cpp
@@ -18,7 +18,7 @@ bool FHoloLensTargetDevice::Deploy(const FString& SourceFolder, FString& OutAppI
 	return Info.bCanDeployTo;
 }
 
-bool FHoloLensTargetDevice::Launch(const FString& AppId, EBuildConfigurations::Type BuildConfiguration, EBuildTargets::Type BuildTarget, const FString& Params, uint32* OutProcessId)
+bool FHoloLensTargetDevice::Launch(const FString& AppId, EBuildConfigurations::Type BuildConfiguration, EBuildTargetType BuildTarget, const FString& Params, uint32* OutProcessId)
 {
 	return false;
 }
diff --git a/Engine/Plugins/Runtime/AR/Microsoft/HoloLensAR/Source/HoloLensTargetPlatform/Private/HoloLensTargetDevice.h b/Engine/Plugins/Runtime/AR/Microsoft/HoloLensAR/Source/HoloLensTargetPlatform/Private/HoloLensTargetDevice.h
--- a/Engine/Plugins/Runtime/AR/Microsoft/HoloLensAR/Source/HoloLensTargetPlatform/Private/HoloLensTargetDevice.h
+++ b/Engine/Plugins/Runtime/AR/Microsoft/HoloLensAR/Source/HoloLensTargetPlatform/Private/HoloLensTargetDevice.h
@@ -103,7 +103,7 @@ public:
 		return Info.HostName == FPlatformProcess::ComputerName();
 	}
 
-	virtual bool Launch(const FString& AppId, EBuildConfigurations::Type BuildConfiguration, EBuildTargets::Type BuildTarget, const FString& Params, uint32* OutProcessId) override;
+	virtual bool Launch(const FString& AppId, EBuildConfigurations::Type BuildConfiguration, EBuildTargetType TargetType, const FString& Params, uint32* OutProcessId) override;
 
 	virtual bool PowerOff(bool Force) override
 	{
diff --git a/Engine/Plugins/Runtime/AR/Microsoft/HoloLensAR/Source/HoloLensTargetPlatform/Private/HoloLensTargetPlatform.cpp b/Engine/Plugins/Runtime/AR/Microsoft/HoloLensAR/Source/HoloLensTargetPlatform/Private/HoloLensTargetPlatform.cpp
--- a/Engine/Plugins/Runtime/AR/Microsoft/HoloLensAR/Source/HoloLensTargetPlatform/Private/HoloLensTargetPlatform.cpp
+++ b/Engine/Plugins/Runtime/AR/Microsoft/HoloLensAR/Source/HoloLensTargetPlatform/Private/HoloLensTargetPlatform.cpp
@@ -148,9 +148,9 @@ void FHoloLensTargetPlatform::OnDeviceDetected(const FHoloLensDeviceInfo& Info)
 	}
 }
 
-bool FHoloLensTargetPlatform::SupportsBuildTarget(EBuildTargets::Type BuildTarget) const
+bool FHoloLensTargetPlatform::SupportsBuildTarget(EBuildTargetType TargetType) const
 {
-	return BuildTarget == EBuildTargets::Game;
+	return TargetType == EBuildTargetType::Game;
 }
 
 bool FHoloLensTargetPlatform::IsSdkInstalled(bool bProjectHasCode, FString& OutDocumentationPath) const
diff --git a/Engine/Plugins/Runtime/AR/Microsoft/HoloLensAR/Source/HoloLensTargetPlatform/Private/HoloLensTargetPlatform.h b/Engine/Plugins/Runtime/AR/Microsoft/HoloLensAR/Source/HoloLensTargetPlatform/Private/HoloLensTargetPlatform.h
--- a/Engine/Plugins/Runtime/AR/Microsoft/HoloLensAR/Source/HoloLensTargetPlatform/Private/HoloLensTargetPlatform.h
+++ b/Engine/Plugins/Runtime/AR/Microsoft/HoloLensAR/Source/HoloLensTargetPlatform/Private/HoloLensTargetPlatform.h
@@ -63,7 +63,7 @@ public:
 
 	virtual bool SupportsFeature(ETargetPlatformFeatures Feature) const override;
 
-	virtual bool SupportsBuildTarget(EBuildTargets::Type BuildTarget) const override;
+	virtual bool SupportsBuildTarget(EBuildTargetType TargetType) const override;
 
 #if WITH_ENGINE
 	virtual void GetReflectionCaptureFormats(TArray<FName>& OutFormats) const override;
diff --git a/Engine/Source/Developer/Android/AndroidTargetPlatform/Private/AndroidTargetDevice.h b/Engine/Source/Developer/Android/AndroidTargetPlatform/Private/AndroidTargetDevice.h
--- a/Engine/Source/Developer/Android/AndroidTargetPlatform/Private/AndroidTargetDevice.h
+++ b/Engine/Source/Developer/Android/AndroidTargetPlatform/Private/AndroidTargetDevice.h
@@ -198,7 +198,7 @@ public:
 		return "Android_Multi";
 	}
 
-	virtual bool Launch(const FString& AppId, EBuildConfigurations::Type BuildConfiguration, EBuildTargets::Type BuildTarget, const FString& Params, uint32* OutProcessId);
+	virtual bool Launch(const FString& AppId, EBuildConfigurations::Type BuildConfiguration, EBuildTargetType TargetType, const FString& Params, uint32* OutProcessId);
 	virtual bool Reboot(bool bReconnect = false) override;
 	virtual bool Run(const FString& ExecutablePath, const FString& Params, uint32* OutProcessId) override;
 	virtual bool TerminateLaunchedProcess(const FString& ProcessIdentifier) override;
diff --git a/Engine/Source/Developer/Android/AndroidTargetPlatform/Private/AndroidTargetDevice.inl b/Engine/Source/Developer/Android/AndroidTargetPlatform/Private/AndroidTargetDevice.inl
--- a/Engine/Source/Developer/Android/AndroidTargetPlatform/Private/AndroidTargetDevice.inl
+++ b/Engine/Source/Developer/Android/AndroidTargetPlatform/Private/AndroidTargetDevice.inl
@@ -48,7 +48,7 @@ inline int32 FAndroidTargetDevice::GetProcessSnapshot( TArray<FTargetDeviceProce
 }
 
 
-inline bool FAndroidTargetDevice::Launch( const FString& AppId, EBuildConfigurations::Type BuildConfiguration, EBuildTargets::Type BuildTarget, const FString& Params, uint32* OutProcessId )
+inline bool FAndroidTargetDevice::Launch( const FString& AppId, EBuildConfigurations::Type BuildConfiguration, EBuildTargetType TargetType, const FString& Params, uint32* OutProcessId )
 {
 	// this isn't used, UAT handles it all
 	return false;
diff --git a/Engine/Source/Developer/BlueprintNativeCodeGen/Private/BlueprintNativeCodeGenUtils.cpp b/Engine/Source/Developer/BlueprintNativeCodeGen/Private/BlueprintNativeCodeGenUtils.cpp
--- a/Engine/Source/Developer/BlueprintNativeCodeGen/Private/BlueprintNativeCodeGenUtils.cpp
+++ b/Engine/Source/Developer/BlueprintNativeCodeGen/Private/BlueprintNativeCodeGenUtils.cpp
@@ -136,7 +136,7 @@ static bool BlueprintNativeCodeGenUtilsImpl::GeneratePluginDescFile(const FBluep
 				// should correspond to UnrealBuildTool::TargetType in TargetRules.cs
 				switch (PlatformInfo.PlatformType)
 				{
-				case PlatformInfo::EPlatformType::Game:
+				case EBuildTargetType::Game:
 					ModuleDesc->WhitelistTargets.AddUnique(TEXT("Game"));
 
 					// Hack to allow clients for PS4/XboxOne (etc.) to build the nativized assets plugin
@@ -147,16 +147,16 @@ static bool BlueprintNativeCodeGenUtilsImpl::GeneratePluginDescFile(const FBluep
 					}
 					break;
 
-				case PlatformInfo::EPlatformType::Client:
+				case EBuildTargetType::Client:
 					ModuleDesc->WhitelistTargets.AddUnique(TEXT("Client"));
 					break;
 
-				case PlatformInfo::EPlatformType::Server:
+				case EBuildTargetType::Server:
 					ModuleDesc->WhitelistTargets.AddUnique(TEXT("Server"));
 					break;
 
-				case PlatformInfo::EPlatformType::Editor:
-					ensureMsgf(PlatformInfo.PlatformType != PlatformInfo::EPlatformType::Editor, TEXT("Nativized Blueprint plugin is for cooked projects only - it isn't supported in editor builds."));
+				case EBuildTargetType::Editor:
+					ensureMsgf(PlatformInfo.PlatformType != EBuildTargetType::Editor, TEXT("Nativized Blueprint plugin is for cooked projects only - it isn't supported in editor builds."));
 					break;
 				};				
 			}
diff --git a/Engine/Source/Developer/DesktopPlatform/Private/PlatformInfo.cpp b/Engine/Source/Developer/DesktopPlatform/Private/PlatformInfo.cpp
--- a/Engine/Source/Developer/DesktopPlatform/Private/PlatformInfo.cpp
+++ b/Engine/Source/Developer/DesktopPlatform/Private/PlatformInfo.cpp
@@ -22,7 +22,7 @@ TArray<FPlatformInfo> AllPlatformInfoArray;
 // @todo platplug: Figure out why this is compiled on target devices
 #if WITH_EDITOR || IS_PROGRAM
 
-void BuildPlatformInfo(const FName& InPlatformInfoName, const FName& InTargetPlatformName, const FText& InDisplayName, const EPlatformType InPlatformType, const EPlatformFlags::Flags InPlatformFlags, const FPlatformIconPaths& InIconPaths, const FString& InUATCommandLine, const FString& InAutoSDKPath, EPlatformSDKStatus InStatus, const FString& InTutorial, bool InEnabled, FString InBinaryFolderName, FString InIniPlatformName, bool InUsesHostCompiler, bool InUATClosesAfterLaunch, bool InIsConfidential, const FName& InUBTTargetId, const FName& InPlatformGroupName, bool InTargetPlatformCanUseCrashReporter)
+void BuildPlatformInfo(const FName& InPlatformInfoName, const FName& InTargetPlatformName, const FText& InDisplayName, const EBuildTargetType InPlatformType, const EPlatformFlags::Flags InPlatformFlags, const FPlatformIconPaths& InIconPaths, const FString& InUATCommandLine, const FString& InAutoSDKPath, EPlatformSDKStatus InStatus, const FString& InTutorial, bool InEnabled, FString InBinaryFolderName, FString InIniPlatformName, bool InUsesHostCompiler, bool InUATClosesAfterLaunch, bool InIsConfidential, const FName& InUBTTargetId, const FName& InPlatformGroupName, bool InTargetPlatformCanUseCrashReporter)
 {
 	FPlatformInfo& PlatformInfo = AllPlatformInfoArray.Emplace_GetRef();
 
@@ -86,7 +86,7 @@ void BuildPlatformInfo(const FName& InPlatformInfoName, const FName& InTargetPla
 void BuildHardcodedPlatforms()
 {
 	// PlatformInfoName									TargetPlatformName			DisplayName														PlatformType			PlatformFlags					IconPaths																																		UATCommandLine										AutoSDKPath			SDKStatus						SDKTutorial																								bEnabledForUse										BinaryFolderName	IniPlatformName		FbUsesHostCompiler		bUATClosesAfterLaunch	bIsConfidential	UBTTargetId (match UBT's UnrealTargetPlatform enum)		bTargetPlatformCanUseCrashReporter
-	BuildPlatformInfo(TEXT("AllDesktop"),				TEXT("AllDesktop"),			LOCTEXT("DesktopTargetPlatDisplay", "Desktop (Win+Mac+Linux)"),	EPlatformType::Game,	EPlatformFlags::None,			FPlatformIconPaths(TEXT("Launcher/Desktop/Platform_Desktop_24x"), TEXT("Launcher/Desktop/Platform_Desktop_128x")),					TEXT(""),											TEXT(""),			EPlatformSDKStatus::Unknown,	TEXT(""),																								PLATFORM_WINDOWS /* see note below */,						TEXT(""),			TEXT(""),			false,					true,					false,			TEXT("AllDesktop"),	TEXT("Desktop"),	true);
+	BuildPlatformInfo(TEXT("AllDesktop"),				TEXT("AllDesktop"),			LOCTEXT("DesktopTargetPlatDisplay", "Desktop (Win+Mac+Linux)"),	EBuildTargetType::Game,	EPlatformFlags::None,			FPlatformIconPaths(TEXT("Launcher/Desktop/Platform_Desktop_24x"), TEXT("Launcher/Desktop/Platform_Desktop_128x")),					TEXT(""),											TEXT(""),			EPlatformSDKStatus::Unknown,	TEXT(""),																								PLATFORM_WINDOWS /* see note below */,						TEXT(""),			TEXT(""),			false,					true,					false,			TEXT("AllDesktop"),	TEXT("Desktop"),	true);
 	// Note: For "AllDesktop" bEnabledForUse value, see SProjectTargetPlatformSettings::Construct !!!! IsAvailableOnWindows || IsAvailableOnMac || IsAvailableOnLinux
 }
 
@@ -144,7 +144,7 @@ void ParseDataDrivenPlatformInfo(const TCHAR* Name, const FConfigSection& Sectio
 	FString TargetPlatformCanUseCrashReporterString = GetSectionString(Section, TEXT("bTargetPlatformCanUseCrashReporter"));
 	bool bTargetPlatformCanUseCrashReporter = TargetPlatformCanUseCrashReporterString.IsEmpty() ? true : GetSectionBool(Section, TEXT("bTargetPlatformCanUseCrashReporter"));
 
-	BuildPlatformInfo(Name, TargetPlatformName, FText::FromString(DisplayName), EPlatformTypeFromString(*PlatformType), ConvertPlatformFlags(PlatformFlags), FPlatformIconPaths(NormalIconPath, LargeIconPath, XLargeIconPath), UATCommandLine,
+	BuildPlatformInfo(Name, TargetPlatformName, FText::FromString(DisplayName), LexFromString(*PlatformType), ConvertPlatformFlags(PlatformFlags), FPlatformIconPaths(NormalIconPath, LargeIconPath, XLargeIconPath), UATCommandLine,
 		AutoSDKPath, PlatformInfo::EPlatformSDKStatus::Unknown, TutorialPath, bIsEnabled, BinariesDirectoryName, IniPlatformName, bUsesHostCompiler, bUATClosesAfterLaunch, bIsConfidential, UBTTargetID, PlatformGroupName, bTargetPlatformCanUseCrashReporter);
 }
 
@@ -302,31 +302,6 @@ FVanillaPlatformEntry BuildPlatformHierarchy(const FName& InPlatformName, const
 	return VanillaPlatformEntry;
 }
 
-EPlatformType EPlatformTypeFromString(const FString& PlatformTypeName)
-{
-	if (FCString::Strcmp(*PlatformTypeName, TEXT("Game")) == 0)
-	{
-		return PlatformInfo::EPlatformType::Game;
-	}
-	else if (FCString::Strcmp(*PlatformTypeName, TEXT("Editor")) == 0)
-	{
-		return PlatformInfo::EPlatformType::Editor;
-	}
-	else if (FCString::Strcmp(*PlatformTypeName, TEXT("Client")) == 0)
-	{
-		return PlatformInfo::EPlatformType::Client;
-	}
-	else if (FCString::Strcmp(*PlatformTypeName, TEXT("Server")) == 0)
-	{
-		return PlatformInfo::EPlatformType::Server;
-	}
-	else
-	{
-		UE_LOG(LogDesktopPlatform, Warning, TEXT("Unable to read Platform Type from %s, defaulting to Game"), *PlatformTypeName);
-		return PlatformInfo::EPlatformType::Game;
-	}
-}
-
 const TArray<FName>& GetAllPlatformGroupNames()
 {
 	return PlatformInfo::AllPlatformGroupNames;
@@ -339,21 +314,4 @@ const TArray<FName>& GetAllVanillaPlatformNames()
 
 } // namespace PlatformInfo
 
-FString LexToString(const PlatformInfo::EPlatformType Value)
-{
-	switch (Value)
-	{
-	case PlatformInfo::EPlatformType::Game:
-		return TEXT("Game");
-	case PlatformInfo::EPlatformType::Editor:
-		return TEXT("Editor");
-	case PlatformInfo::EPlatformType::Client:
-		return TEXT("Client");
-	case PlatformInfo::EPlatformType::Server:
-		return TEXT("Server");
-	}
-
-	return TEXT("");
-}
-
 #undef LOCTEXT_NAMESPACE
diff --git a/Engine/Source/Developer/DesktopPlatform/Public/PlatformInfo.h b/Engine/Source/Developer/DesktopPlatform/Public/PlatformInfo.h
--- a/Engine/Source/Developer/DesktopPlatform/Public/PlatformInfo.h
+++ b/Engine/Source/Developer/DesktopPlatform/Public/PlatformInfo.h
@@ -6,22 +6,6 @@
 
 namespace PlatformInfo
 {
-	/** The target type of the given platform */
-	enum class EPlatformType : uint8
-	{
-		/** The platform targets cooked monolithic game executables */
-		Game,
-
-		/** The platform targets uncooked modular editor executables and DLLs */
-		Editor,
-
-		/** The platform targets cooked monolithic game client executables (but no server code) */
-		Client,
-
-		/** The platform targets cooked monolithic game server executables (but no client code) */
-		Server,
-	};
-
 	/** Available icon sizes (see FPlatformIconPaths) */
 	enum class EPlatformIconSize : uint8
 	{
@@ -134,7 +118,7 @@ namespace PlatformInfo
 		FText DisplayName;
 
 		/** Type of this platform */
-		EPlatformType PlatformType;
+		EBuildTargetType PlatformType;
 
 		/** Flags for this platform */
 		EPlatformFlags::Flags PlatformFlags;
@@ -298,13 +282,6 @@ namespace PlatformInfo
 	*/
 	DESKTOPPLATFORM_API void UpdatePlatformDisplayName(FString InPlatformName, FText InDisplayName);
 
-	/**
-	* Returns an EPlatformType value from a string representation.
-	* @param PlatformTypeName The string to get the EPlatformType for.
-	* @return An EPlatformType value.
-	*/
-	DESKTOPPLATFORM_API EPlatformType EPlatformTypeFromString(const FString& PlatformTypeName);
-
 	/**
 	* Returns a list of all defined Platform Groups, excluding None.
     * Used to to present a list in the Per-Platform Properties UI.
@@ -319,10 +296,3 @@ namespace PlatformInfo
 	*/
 	DESKTOPPLATFORM_API const TArray<FName>& GetAllVanillaPlatformNames();
 }
-
-DESKTOPPLATFORM_API FString LexToString(const PlatformInfo::EPlatformType Value);
-
-inline void LexFromString(PlatformInfo::EPlatformType& OutValue, const TCHAR* Buffer)
-{
-	OutValue = PlatformInfo::EPlatformTypeFromString(FString(Buffer));
-}
diff --git a/Engine/Source/Developer/HTML5/HTML5TargetPlatform/Private/HTML5TargetDevice.cpp b/Engine/Source/Developer/HTML5/HTML5TargetPlatform/Private/HTML5TargetDevice.cpp
--- a/Engine/Source/Developer/HTML5/HTML5TargetPlatform/Private/HTML5TargetDevice.cpp
+++ b/Engine/Source/Developer/HTML5/HTML5TargetPlatform/Private/HTML5TargetDevice.cpp
@@ -48,7 +48,7 @@ int32 FHTML5TargetDevice::GetProcessSnapshot( TArray<FTargetDeviceProcessInfo>&
 }
 
 
-bool FHTML5TargetDevice::Launch( const FString& AppId, EBuildConfigurations::Type BuildConfiguration, EBuildTargets::Type BuildTarget, const FString& Params, uint32* OutProcessId )
+bool FHTML5TargetDevice::Launch( const FString& AppId, EBuildConfigurations::Type BuildConfiguration, EBuildTargetType TargetType, const FString& Params, uint32* OutProcessId )
 {
 	return Run(AppId, Params, OutProcessId);
 }
diff --git a/Engine/Source/Developer/HTML5/HTML5TargetPlatform/Private/HTML5TargetDevice.h b/Engine/Source/Developer/HTML5/HTML5TargetPlatform/Private/HTML5TargetDevice.h
--- a/Engine/Source/Developer/HTML5/HTML5TargetPlatform/Private/HTML5TargetDevice.h
+++ b/Engine/Source/Developer/HTML5/HTML5TargetPlatform/Private/HTML5TargetDevice.h
@@ -90,7 +90,7 @@ public:
 		return true;
 	}
 
-	virtual bool Launch( const FString& AppId, EBuildConfigurations::Type BuildConfiguration, EBuildTargets::Type BuildTarget, const FString& Params, uint32* OutProcessId ) override;
+	virtual bool Launch( const FString& AppId, EBuildConfigurations::Type BuildConfiguration, EBuildTargetType TargetType, const FString& Params, uint32* OutProcessId ) override;
 
 	virtual bool PowerOff( bool Force ) override;
 
diff --git a/Engine/Source/Developer/IOS/IOSTargetPlatform/Private/IOSTargetDevice.cpp b/Engine/Source/Developer/IOS/IOSTargetPlatform/Private/IOSTargetDevice.cpp
--- a/Engine/Source/Developer/IOS/IOSTargetPlatform/Private/IOSTargetDevice.cpp
+++ b/Engine/Source/Developer/IOS/IOSTargetPlatform/Private/IOSTargetDevice.cpp
@@ -329,7 +329,7 @@ bool FIOSTargetDevice::IsDefault() const
 	return true;
 }
 	
-bool FIOSTargetDevice::Launch(const FString& InAppId, EBuildConfigurations::Type InBuildConfiguration, EBuildTargets::Type BuildTarget, const FString& Params, uint32* OutProcessId)
+bool FIOSTargetDevice::Launch(const FString& InAppId, EBuildConfigurations::Type InBuildConfiguration, EBuildTargetType TargetType, const FString& Params, uint32* OutProcessId)
 {
 #if !PLATFORM_MAC
 	MessageEndpoint->Send(new FIOSLaunchDaemonLaunchApp(InAppId, Params), DeviceEndpoint);
diff --git a/Engine/Source/Developer/IOS/IOSTargetPlatform/Private/IOSTargetDevice.h b/Engine/Source/Developer/IOS/IOSTargetPlatform/Private/IOSTargetDevice.h
--- a/Engine/Source/Developer/IOS/IOSTargetPlatform/Private/IOSTargetDevice.h
+++ b/Engine/Source/Developer/IOS/IOSTargetPlatform/Private/IOSTargetDevice.h
@@ -118,7 +118,7 @@ public:
 	virtual const class ITargetPlatform& GetTargetPlatform() const override;
 	virtual bool IsConnected() override;
 	virtual bool IsDefault() const override;
-	virtual bool Launch(const FString& InAppId, EBuildConfigurations::Type InBuildConfiguration, EBuildTargets::Type BuildTarget, const FString& Params, uint32* OutProcessId) override;
+	virtual bool Launch(const FString& InAppId, EBuildConfigurations::Type InBuildConfiguration, EBuildTargetType TargetType, const FString& Params, uint32* OutProcessId) override;
 	virtual bool PowerOff(bool Force) override;
 	virtual bool PowerOn() override;
 	virtual bool Reboot(bool bReconnect = false) override;
diff --git a/Engine/Source/Developer/Linux/LinuxTargetPlatform/Private/LinuxTargetDevice.h b/Engine/Source/Developer/Linux/LinuxTargetPlatform/Private/LinuxTargetDevice.h
--- a/Engine/Source/Developer/Linux/LinuxTargetPlatform/Private/LinuxTargetDevice.h
+++ b/Engine/Source/Developer/Linux/LinuxTargetPlatform/Private/LinuxTargetDevice.h
@@ -143,22 +143,22 @@ public:
 		return false;
 	}
 
-	virtual bool Launch( const FString& AppId, EBuildConfigurations::Type BuildConfiguration, EBuildTargets::Type BuildTarget, const FString& Params, uint32* OutProcessId ) override
+	virtual bool Launch( const FString& AppId, EBuildConfigurations::Type BuildConfiguration, EBuildTargetType TargetType, const FString& Params, uint32* OutProcessId ) override
 	{
 #if PLATFORM_LINUX	// if running natively, support launching in place
 		// build executable path
 		FString PlatformName = TEXT("Linux");
 		FString ExecutablePath = FPaths::EngineIntermediateDir() / TEXT("Devices") / PlatformName / TEXT("Engine") / TEXT("Binaries") / PlatformName;
 
-		if (BuildTarget == EBuildTargets::Game)
+		if (BuildTarget == EBuildTargetType::Game)
 		{
 			ExecutablePath /= TEXT("UE4Game");
 		}
-		else if (BuildTarget == EBuildTargets::Server)
+		else if (BuildTarget == EBuildTargetType::Server)
 		{
 			ExecutablePath /= TEXT("UE4Server");
 		}
-		else if (BuildTarget == EBuildTargets::Editor)
+		else if (BuildTarget == EBuildTargetType::Editor)
 		{
 			ExecutablePath /= TEXT("UE4Editor");
 		}
diff --git a/Engine/Source/Developer/Mac/MacTargetPlatform/Private/LocalMacTargetDevice.h b/Engine/Source/Developer/Mac/MacTargetPlatform/Private/LocalMacTargetDevice.h
--- a/Engine/Source/Developer/Mac/MacTargetPlatform/Private/LocalMacTargetDevice.h
+++ b/Engine/Source/Developer/Mac/MacTargetPlatform/Private/LocalMacTargetDevice.h
@@ -106,7 +106,7 @@ public:
 		return true;
 	}
 
-	virtual bool Launch( const FString& AppId, EBuildConfigurations::Type BuildConfiguration, EBuildTargets::Type Target, const FString& Params, uint32* OutProcessId ) override
+	virtual bool Launch( const FString& AppId, EBuildConfigurations::Type BuildConfiguration, EBuildTargetType TargetType, const FString& Params, uint32* OutProcessId ) override
 	{
 		// build executable path
 		FString PlatformName = TEXT("Mac");
diff --git a/Engine/Source/Developer/TargetDeviceServices/Private/Services/TargetDeviceService.cpp b/Engine/Source/Developer/TargetDeviceServices/Private/Services/TargetDeviceService.cpp
--- a/Engine/Source/Developer/TargetDeviceServices/Private/Services/TargetDeviceService.cpp
+++ b/Engine/Source/Developer/TargetDeviceServices/Private/Services/TargetDeviceService.cpp
@@ -463,7 +463,7 @@ void FTargetDeviceService::HandleLaunchAppMessage(const FTargetDeviceServiceLaun
 	if (TargetDevice.IsValid())
 	{
 		uint32 ProcessId;
-		bool Succeeded = TargetDevice->Launch(Message.AppID, (EBuildConfigurations::Type)Message.BuildConfiguration, EBuildTargets::Game, Message.Params, &ProcessId);
+		bool Succeeded = TargetDevice->Launch(Message.AppID, (EBuildConfigurations::Type)Message.BuildConfiguration, EBuildTargetType::Game, Message.Params, &ProcessId);
 
 		if (MessageEndpoint.IsValid())
 		{
diff --git a/Engine/Source/Developer/TargetPlatform/Private/InstalledPlatformInfo.cpp b/Engine/Source/Developer/TargetPlatform/Private/InstalledPlatformInfo.cpp
--- a/Engine/Source/Developer/TargetPlatform/Private/InstalledPlatformInfo.cpp
+++ b/Engine/Source/Developer/TargetPlatform/Private/InstalledPlatformInfo.cpp
@@ -76,10 +76,10 @@ void FInstalledPlatformInfo::ParsePlatformConfiguration(FString PlatformConfigur
 	}
 
 	FString PlatformTypeName;
-	PlatformInfo::EPlatformType PlatformType = PlatformInfo::EPlatformType::Game;
+	EBuildTargetType PlatformType = EBuildTargetType::Game;
 	if (FParse::Value(*PlatformConfiguration, TEXT("PlatformType="), PlatformTypeName))
 	{
-		PlatformType = PlatformInfo::EPlatformTypeFromString(PlatformTypeName);
+		PlatformType = LexFromString(*PlatformTypeName);
 	}
 
 	FString Architecture;
@@ -162,12 +162,12 @@ bool FInstalledPlatformInfo::CanDisplayPlatform(const FString& PlatformName, EPr
 	);
 }
 
-bool FInstalledPlatformInfo::IsValidPlatformType(PlatformInfo::EPlatformType PlatformType) const
+bool FInstalledPlatformInfo::IsValidTargetType(EBuildTargetType TargetType) const
 {
 	return ContainsValidConfiguration(
-		[PlatformType](const FInstalledPlatformConfiguration& CurConfig)
+		[TargetType](const FInstalledPlatformConfiguration& CurConfig)
 		{
-			return CurConfig.PlatformType == PlatformType;
+			return CurConfig.PlatformType == TargetType;
 		}
 	);
 }
diff --git a/Engine/Source/Developer/TargetPlatform/Public/Common/TargetPlatformBase.h b/Engine/Source/Developer/TargetPlatform/Public/Common/TargetPlatformBase.h
--- a/Engine/Source/Developer/TargetPlatform/Public/Common/TargetPlatformBase.h
+++ b/Engine/Source/Developer/TargetPlatform/Public/Common/TargetPlatformBase.h
@@ -208,9 +208,9 @@ public:
 		return TPlatformProperties::RequiresUserCredentials();
 	}
 
-	virtual bool SupportsBuildTarget( EBuildTargets::Type BuildTarget ) const override
+	virtual bool SupportsBuildTarget( EBuildTargetType TargetType ) const override
 	{
-		return TPlatformProperties::SupportsBuildTarget(BuildTarget);
+		return TPlatformProperties::SupportsBuildTarget(TargetType);
 	}
 
 	virtual bool SupportsAutoSDK() const override
diff --git a/Engine/Source/Developer/TargetPlatform/Public/InstalledPlatformInfo.h b/Engine/Source/Developer/TargetPlatform/Public/InstalledPlatformInfo.h
--- a/Engine/Source/Developer/TargetPlatform/Public/InstalledPlatformInfo.h
+++ b/Engine/Source/Developer/TargetPlatform/Public/InstalledPlatformInfo.h
@@ -5,11 +5,6 @@
 #include "CoreMinimal.h"
 
 // Forward declaration
-namespace PlatformInfo
-{
-	enum class EPlatformType : uint8;
-}
-
 enum class EProjectType : uint8
 {
 	Unknown,
@@ -32,7 +27,7 @@ struct FInstalledPlatformConfiguration
 	FString PlatformName;
 
 	/** Type of Platform for this combination */
-	PlatformInfo::EPlatformType PlatformType;
+	EBuildTargetType PlatformType;
 
 	/** Name of the Architecture for this combination */
 	FString Architecture;
@@ -83,9 +78,9 @@ public:
 	bool CanDisplayPlatform(const FString& PlatformName, EProjectType ProjectType) const;
 
 	/**
-	 * Queries whether a platform type is valid for any configuration
+	 * Queries whether a target type is valid for any configuration
 	 */
-	bool IsValidPlatformType(PlatformInfo::EPlatformType PlatformType) const;
+	bool IsValidTargetType(EBuildTargetType TargetType) const;
 
 	/**
 	 * Queries whether a platform architecture is valid for any configuration
diff --git a/Engine/Source/Developer/TargetPlatform/Public/Interfaces/ITargetDevice.h b/Engine/Source/Developer/TargetPlatform/Public/Interfaces/ITargetDevice.h
--- a/Engine/Source/Developer/TargetPlatform/Public/Interfaces/ITargetDevice.h
+++ b/Engine/Source/Developer/TargetPlatform/Public/Interfaces/ITargetDevice.h
@@ -307,12 +307,12 @@ public:
 	 *
 	 * @param AppId The identifier of the application to launch (as returned by the Deploy() method).
 	 * @param BuildConfiguration The build configuration to launch.
-	 * @param BuildTarget The build target type to launch
+	 * @param TargetType The target type to launch
 	 * @param Params The command line parameters to launch with.
 	 * @param OutProcessId Will hold the identifier of the created process (can be NULL).
 	 * @return true on success, false otherwise.
 	 */
-	virtual bool Launch( const FString& AppId, EBuildConfigurations::Type BuildConfiguration, EBuildTargets::Type BuildTarget, const FString& Params, uint32* OutProcessId ) = 0;
+	virtual bool Launch( const FString& AppId, EBuildConfigurations::Type BuildConfiguration, EBuildTargetType TargetType, const FString& Params, uint32* OutProcessId ) = 0;
 
 	/**
 	 * Powers off the device.
diff --git a/Engine/Source/Developer/TargetPlatform/Public/Interfaces/ITargetPlatform.h b/Engine/Source/Developer/TargetPlatform/Public/Interfaces/ITargetPlatform.h
--- a/Engine/Source/Developer/TargetPlatform/Public/Interfaces/ITargetPlatform.h
+++ b/Engine/Source/Developer/TargetPlatform/Public/Interfaces/ITargetPlatform.h
@@ -323,10 +323,10 @@ public:
 	/**
 	 * Checks whether this platform supports the specified build target, i.e. Game or Editor.
 	 *
-	 * @param BuildTarget The build target to check.
+	 * @param TargetType The target type to check.
 	 * @return true if the build target is supported, false otherwise.
 	 */
-	virtual bool SupportsBuildTarget( EBuildTargets::Type BuildTarget ) const = 0;
+	virtual bool SupportsBuildTarget( EBuildTargetType TargetType ) const = 0;
 
 	/**
 	 * Checks whether the target platform supports the specified feature.
diff --git a/Engine/Source/Developer/Windows/WindowsTargetPlatform/Private/LocalPcTargetDevice.h b/Engine/Source/Developer/Windows/WindowsTargetPlatform/Private/LocalPcTargetDevice.h
--- a/Engine/Source/Developer/Windows/WindowsTargetPlatform/Private/LocalPcTargetDevice.h
+++ b/Engine/Source/Developer/Windows/WindowsTargetPlatform/Private/LocalPcTargetDevice.h
@@ -237,21 +237,21 @@ public:
 		return true;
 	}
 
-	virtual bool Launch( const FString& AppId, EBuildConfigurations::Type BuildConfiguration, EBuildTargets::Type BuildTarget, const FString& Params, uint32* OutProcessId ) override
+	virtual bool Launch( const FString& AppId, EBuildConfigurations::Type BuildConfiguration, EBuildTargetType TargetType, const FString& Params, uint32* OutProcessId ) override
 	{
 		// build executable path
 		FString PlatformName = WIN64 ? TEXT("Win64") : TEXT("Win32");
 		FString ExecutablePath = FPaths::EngineIntermediateDir() / TEXT("Devices") / PlatformName / TEXT("Engine") / TEXT("Binaries") / PlatformName;
 		
-		if (BuildTarget == EBuildTargets::Game)
+		if (TargetType == EBuildTargetType::Game)
 		{
 			ExecutablePath /= TEXT("UE4Game");
 		}
-		else if (BuildTarget == EBuildTargets::Server)
+		else if (TargetType == EBuildTargetType::Server)
 		{
 			ExecutablePath /= TEXT("UE4Game");
 		}
-		else if (BuildTarget == EBuildTargets::Editor)
+		else if (TargetType == EBuildTargetType::Editor)
 		{
 			ExecutablePath /= TEXT("UE4Editor");
 		}
diff --git a/Engine/Source/Editor/GameProjectGeneration/Private/SProjectBrowser.cpp b/Engine/Source/Editor/GameProjectGeneration/Private/SProjectBrowser.cpp
--- a/Engine/Source/Editor/GameProjectGeneration/Private/SProjectBrowser.cpp
+++ b/Engine/Source/Editor/GameProjectGeneration/Private/SProjectBrowser.cpp
@@ -867,7 +867,7 @@ FReply SProjectBrowser::FindProjects()
 				TArray<FName> TargetPlatforms;
 				for(const PlatformInfo::FPlatformInfo& PlatformInfo : PlatformInfo::GetPlatformInfoArray())
 				{
-					if(PlatformInfo.IsVanilla() && PlatformInfo.PlatformType == PlatformInfo::EPlatformType::Game && ProjectStatus.IsTargetPlatformSupported(PlatformInfo.PlatformInfoName))
+					if(PlatformInfo.IsVanilla() && PlatformInfo.PlatformType == EBuildTargetType::Game && ProjectStatus.IsTargetPlatformSupported(PlatformInfo.PlatformInfoName))
 					{
 						TargetPlatforms.Add(PlatformInfo.PlatformInfoName);
 					}
diff --git a/Engine/Source/Editor/MainFrame/Private/Menus/CookContentMenu.h b/Engine/Source/Editor/MainFrame/Private/Menus/CookContentMenu.h
--- a/Engine/Source/Editor/MainFrame/Private/Menus/CookContentMenu.h
+++ b/Engine/Source/Editor/MainFrame/Private/Menus/CookContentMenu.h
@@ -55,7 +55,7 @@ public:
 			check(VanillaPlatform.PlatformInfo->IsVanilla());
 
 			// Only care about game targets
-			if (VanillaPlatform.PlatformInfo->PlatformType != PlatformInfo::EPlatformType::Game || !VanillaPlatform.PlatformInfo->bEnabledForUse || !FInstalledPlatformInfo::Get().CanDisplayPlatform(VanillaPlatform.PlatformInfo->BinaryFolderName, ProjectType))
+			if (VanillaPlatform.PlatformInfo->PlatformType != EBuildTargetType::Game || !VanillaPlatform.PlatformInfo->bEnabledForUse || !FInstalledPlatformInfo::Get().CanDisplayPlatform(VanillaPlatform.PlatformInfo->BinaryFolderName, ProjectType))
 			{
 				continue;
 			}
diff --git a/Engine/Source/Editor/MainFrame/Private/Menus/PackageProjectMenu.h b/Engine/Source/Editor/MainFrame/Private/Menus/PackageProjectMenu.h
--- a/Engine/Source/Editor/MainFrame/Private/Menus/PackageProjectMenu.h
+++ b/Engine/Source/Editor/MainFrame/Private/Menus/PackageProjectMenu.h
@@ -55,7 +55,7 @@ public:
 			check(VanillaPlatform.PlatformInfo->IsVanilla());
 
 			// Only care about game targets
-			if (VanillaPlatform.PlatformInfo->PlatformType != PlatformInfo::EPlatformType::Game || !VanillaPlatform.PlatformInfo->bEnabledForUse || !FInstalledPlatformInfo::Get().CanDisplayPlatform(VanillaPlatform.PlatformInfo->BinaryFolderName, ProjectType))
+			if (VanillaPlatform.PlatformInfo->PlatformType != EBuildTargetType::Game || !VanillaPlatform.PlatformInfo->bEnabledForUse || !FInstalledPlatformInfo::Get().CanDisplayPlatform(VanillaPlatform.PlatformInfo->BinaryFolderName, ProjectType))
 			{
 				continue;
 			}
@@ -153,7 +153,7 @@ protected:
 	{
 		for (const PlatformInfo::FPlatformInfo* SubPlatformInfo : SubPlatformInfos)
 		{
-			if (SubPlatformInfo->PlatformType != PlatformInfo::EPlatformType::Game)
+			if (SubPlatformInfo->PlatformType != EBuildTargetType::Game)
 			{
 				continue;
 			}
@@ -190,7 +190,7 @@ protected:
 
 		// Add the Client configurations if there is a {ProjectName}Client.Target.cs file.
 		TArray<FString> ClientTargetFileNames;
-		if (FInstalledPlatformInfo::Get().IsValidPlatformType(PlatformInfo::EPlatformType::Client))
+		if (FInstalledPlatformInfo::Get().IsValidTargetType(EBuildTargetType::Client))
 		{
 			IFileManager::Get().FindFiles(ClientTargetFileNames, *(FPaths::GameSourceDir() / TEXT("*client.target.cs")), true, false);
 		}
diff --git a/Engine/Source/Editor/ProjectTargetPlatformEditor/Private/Widgets/SProjectTargetPlatformSettings.cpp b/Engine/Source/Editor/ProjectTargetPlatformEditor/Private/Widgets/SProjectTargetPlatformSettings.cpp
--- a/Engine/Source/Editor/ProjectTargetPlatformEditor/Private/Widgets/SProjectTargetPlatformSettings.cpp
+++ b/Engine/Source/Editor/ProjectTargetPlatformEditor/Private/Widgets/SProjectTargetPlatformSettings.cpp
@@ -24,7 +24,7 @@ void SProjectTargetPlatformSettings::Construct(const FArguments& InArgs)
 	// We show all of the platforms regardless of whether we have an SDK installed for them or not
 	for(const PlatformInfo::FPlatformInfo& PlatformInfo : PlatformInfo::GetPlatformInfoArray())
 	{
-		if(PlatformInfo.IsVanilla() && PlatformInfo.PlatformType == PlatformInfo::EPlatformType::Game)
+		if(PlatformInfo.IsVanilla() && PlatformInfo.PlatformType == EBuildTargetType::Game)
 		{
 #if !PLATFORM_WINDOWS
 			// @todo AllDesktop now only works on Windows (it can compile D3D shaders, and it can remote compile Metal shaders)
diff --git a/Engine/Source/Editor/UnrealEd/Private/Kismet2/DebuggerCommands.cpp b/Engine/Source/Editor/UnrealEd/Private/Kismet2/DebuggerCommands.cpp
--- a/Engine/Source/Editor/UnrealEd/Private/Kismet2/DebuggerCommands.cpp
+++ b/Engine/Source/Editor/UnrealEd/Private/Kismet2/DebuggerCommands.cpp
@@ -975,7 +975,7 @@ static void MakeAllDevicesSubMenu(FMenuBuilder& InMenuBuilder, const PlatformInf
 		// for an aggregate (All_<platform>_devices_on_<host>) proxy, allow only the "Android_<texture_compression>" variants
 		const PlatformInfo::FPlatformInfo* platformInfo = PlatformInfo::FindPlatformInfo(Variant);
 		if (DeviceProxy->IsAggregated() && platformInfo != NULL &&
-			(Variant == platformInfo->VanillaPlatformName || platformInfo->PlatformType != PlatformInfo::EPlatformType::Game))
+			(Variant == platformInfo->VanillaPlatformName || platformInfo->PlatformType != EBuildTargetType::Game))
 		{
 			continue;
 		}
@@ -1063,7 +1063,7 @@ TSharedRef< SWidget > FPlayWorldCommands::GenerateLaunchMenuContent( TSharedRef<
 		for (const PlatformInfo::FVanillaPlatformEntry& VanillaPlatform : VanillaPlatforms)
 		{
 			// for the Editor we are only interested in launching standalone games
-			if (VanillaPlatform.PlatformInfo->PlatformType != PlatformInfo::EPlatformType::Game || !VanillaPlatform.PlatformInfo->bEnabledForUse || !FInstalledPlatformInfo::Get().CanDisplayPlatform(VanillaPlatform.PlatformInfo->BinaryFolderName, ProjectType))
+			if (VanillaPlatform.PlatformInfo->PlatformType != EBuildTargetType::Game || !VanillaPlatform.PlatformInfo->bEnabledForUse || !FInstalledPlatformInfo::Get().CanDisplayPlatform(VanillaPlatform.PlatformInfo->BinaryFolderName, ProjectType))
 			{
 				continue;
 			}
diff --git a/Engine/Source/Runtime/Core/Private/GenericPlatform/GenericPlatformMisc.cpp b/Engine/Source/Runtime/Core/Private/GenericPlatform/GenericPlatformMisc.cpp
--- a/Engine/Source/Runtime/Core/Private/GenericPlatform/GenericPlatformMisc.cpp
+++ b/Engine/Source/Runtime/Core/Private/GenericPlatform/GenericPlatformMisc.cpp
@@ -131,48 +131,64 @@ namespace EBuildConfigurations
 }
 
 
-/* EBuildConfigurations interface
+/* EBuildTargetType functions
  *****************************************************************************/
 
-namespace EBuildTargets
+EBuildTargetType LexFromString(const TCHAR* Type)
 {
-	EBuildTargets::Type FromString( const FString& Target )
+	if (FCString::Strcmp(Type, TEXT("Editor")) == 0)
 	{
-		if (FCString::Strcmp(*Target, TEXT("Editor")) == 0)
-		{
-			return Editor;
-		}
-		else if (FCString::Strcmp(*Target, TEXT("Game")) == 0)
-		{
-			return Game;
-		}
-		else if (FCString::Strcmp(*Target, TEXT("Server")) == 0)
-		{
-			return Server;
-		}
-
-		return Unknown;
+		return EBuildTargetType::Editor;
 	}
-
-	const TCHAR* ToString( EBuildTargets::Type Target )
+	else if (FCString::Strcmp(Type, TEXT("Game")) == 0)
 	{
-		switch (Target)
-		{
-			case Editor:
-				return TEXT("Editor");
-
-			case Game:
-				return TEXT("Game");
+		return EBuildTargetType::Game;
+	}
+	else if (FCString::Strcmp(Type, TEXT("Server")) == 0)
+	{
+		return EBuildTargetType::Server;
+	}
+	else if (FCString::Strcmp(Type, TEXT("Client")) == 0)
+	{
+		return EBuildTargetType::Client;
+	}
+	else if (FCString::Strcmp(Type, TEXT("Program")) == 0)
+	{
+		return EBuildTargetType::Program;
+	}
 
-			case Server:
-				return TEXT("Server");
+	return EBuildTargetType::Unknown;
+}
 
-			default:
-				return TEXT("Unknown");
-		}
+const TCHAR* LexToString(EBuildTargetType Type)
+{
+	switch (Type)
+	{
+	case EBuildTargetType::Editor:
+		return TEXT("Editor");
+	case EBuildTargetType::Game:
+		return TEXT("Game");
+	case EBuildTargetType::Server:
+		return TEXT("Server");
+	case EBuildTargetType::Client:
+		return TEXT("Client");
+	case EBuildTargetType::Program:
+		return TEXT("Program");
+	default:
+		return TEXT("Unknown");
 	}
 }
 
+EBuildTargetType EBuildTargets::FromString(const FString& Target)
+{
+	return LexFromString(*Target);
+}
+
+const TCHAR* EBuildTargets::ToString(EBuildTargetType Target)
+{
+	return LexToString(Target);
+}
+
 FString FSHA256Signature::ToString() const
 {
 	FString LocalHashStr;
diff --git a/Engine/Source/Runtime/Core/Public/Android/AndroidPlatformProperties.h b/Engine/Source/Runtime/Core/Public/Android/AndroidPlatformProperties.h
--- a/Engine/Source/Runtime/Core/Public/Android/AndroidPlatformProperties.h
+++ b/Engine/Source/Runtime/Core/Public/Android/AndroidPlatformProperties.h
@@ -49,9 +49,9 @@ struct FAndroidPlatformProperties
 		return true;
 	}
 
-	static FORCEINLINE bool SupportsBuildTarget( EBuildTargets::Type BuildTarget )
+	static FORCEINLINE bool SupportsBuildTarget( EBuildTargetType TargetType )
 	{
-		return (BuildTarget == EBuildTargets::Game);
+		return (TargetType == EBuildTargetType::Game);
 	}
 
 	static FORCEINLINE bool SupportsAutoSDK()
diff --git a/Engine/Source/Runtime/Core/Public/GenericPlatform/GenericPlatformMisc.h b/Engine/Source/Runtime/Core/Public/GenericPlatform/GenericPlatformMisc.h
--- a/Engine/Source/Runtime/Core/Public/GenericPlatform/GenericPlatformMisc.h
+++ b/Engine/Source/Runtime/Core/Public/GenericPlatform/GenericPlatformMisc.h
@@ -82,45 +82,70 @@ FORCEINLINE const TCHAR* LexToString(EBuildConfigurations::Type Configuration)
 	return EBuildConfigurations::ToString(Configuration);
 }
 
+/**
+ * Enumerates build target types.
+ */
+enum class EBuildTargetType
+{
+	/** Unknown build target. */
+	Unknown,
+
+	/** Game target. */
+	Game,
+
+	/** Server target. */
+	Server,
+
+	/** Client target. */
+	Client,
+
+	/** Editor target. */
+	Editor,
+
+	/** Program target. */
+	Program,
+};
+
+/**
+ * Returns the string representation of the specified EBuildTarget value.
+ *
+ * @param Target The value to get the string for.
+ * @return The string representation.
+ */
+CORE_API EBuildTargetType LexFromString(const TCHAR* Type);
+
+/**
+ * Returns the string representation of the specified EBuildTargetType value.
+ *
+ * @param Target The string to get the EBuildTargetType for.
+ * @return An EBuildTarget::Type value.
+ */
+CORE_API const TCHAR* LexToString(EBuildTargetType Type);
 
 namespace EBuildTargets
 {
-	/**
-	 * Enumerates build targets.
-	 */
-	enum Type
-	{
-		/** Unknown build target. */
-		Unknown,
+	UE_DEPRECATED(4.24, "EBuildTargets::Type is deprecated. Use EBuildTargetType instead.")
+	typedef EBuildTargetType Type;
 
-		/** Editor target. */
-		Editor,
+	UE_DEPRECATED(4.24, "EBuildTargets::Unknown is deprecated. Use EBuildTargetType::Unknown instead.")
+	static const EBuildTargetType Unknown = EBuildTargetType::Unknown;
 
-		/** Game target. */
-		Game,
+	UE_DEPRECATED(4.24, "EBuildTargets::Editor is deprecated. Use EBuildTargetType::Unknown instead.")
+	static const EBuildTargetType Editor = EBuildTargetType::Editor;
 
-		/** Server target. */
-		Server
-	};
+	UE_DEPRECATED(4.24, "EBuildTargets::Game is deprecated. Use EBuildTargetType::Unknown instead.")
+	static const EBuildTargetType Game = EBuildTargetType::Game;
 
-	/**
-	 * Returns the string representation of the specified EBuildTarget value.
-	 *
-	 * @param Target The string to get the EBuildTarget::Type for.
-	 * @return An EBuildTarget::Type value.
-	 */
-	CORE_API EBuildTargets::Type FromString( const FString& Target );
+	UE_DEPRECATED(4.24, "EBuildTargets::Server is deprecated. Use EBuildTargetType::Unknown instead.")
+	static const EBuildTargetType Server = EBuildTargetType::Server;
 
-	/**
-	 * Returns the string representation of the specified EBuildTarget value.
-	 *
-	 * @param Target The value to get the string for.
-	 * @return The string representation.
-	 */
-	CORE_API const TCHAR* ToString( EBuildTargets::Type Target );
+	UE_DEPRECATED(4.24, "EBuildTargets::FromString is deprecated. Use LexFromString() instead.")
+	CORE_API EBuildTargetType FromString( const FString& Target );
+	
+	UE_DEPRECATED(4.24, "EBuildTargets::FromString is deprecated. Use LexFromString() instead.")
+	CORE_API const TCHAR* ToString(EBuildTargetType Target);
 }
 
-
 /**
  * Enumerates the modes a convertible laptop can be in.
  */
diff --git a/Engine/Source/Runtime/Core/Public/GenericPlatform/GenericPlatformProperties.h b/Engine/Source/Runtime/Core/Public/GenericPlatform/GenericPlatformProperties.h
--- a/Engine/Source/Runtime/Core/Public/GenericPlatform/GenericPlatformProperties.h
+++ b/Engine/Source/Runtime/Core/Public/GenericPlatform/GenericPlatformProperties.h
@@ -142,10 +142,10 @@ struct FGenericPlatformProperties
 	/**
 	 * Checks whether the specified build target is supported.
 	 *
-	 * @param BuildTarget The build target to check.
+	 * @param TargetType The build target to check.
 	 * @return true if the build target is supported, false otherwise.
 	 */
-	static FORCEINLINE bool SupportsBuildTarget( EBuildTargets::Type BuildTarget )
+	static FORCEINLINE bool SupportsBuildTarget( EBuildTargetType TargetType )
 	{
 		return true;
 	}
diff --git a/Engine/Source/Runtime/Core/Public/HTML5/HTML5PlatformProperties.h b/Engine/Source/Runtime/Core/Public/HTML5/HTML5PlatformProperties.h
--- a/Engine/Source/Runtime/Core/Public/HTML5/HTML5PlatformProperties.h
+++ b/Engine/Source/Runtime/Core/Public/HTML5/HTML5PlatformProperties.h
@@ -44,9 +44,9 @@ struct FHTML5PlatformProperties
 		return true;
 	}
 
-	static FORCEINLINE bool SupportsBuildTarget( EBuildTargets::Type BuildTarget )
+	static FORCEINLINE bool SupportsBuildTarget( EBuildTargetType TargetType )
 	{
-		return (BuildTarget == EBuildTargets::Game);
+		return (TargetType == EBuildTargetType::Game);
 	}
 
 	static FORCEINLINE bool SupportsHighQualityLightmaps()
diff --git a/Engine/Source/Runtime/Core/Public/IOS/IOSPlatformProperties.h b/Engine/Source/Runtime/Core/Public/IOS/IOSPlatformProperties.h
--- a/Engine/Source/Runtime/Core/Public/IOS/IOSPlatformProperties.h
+++ b/Engine/Source/Runtime/Core/Public/IOS/IOSPlatformProperties.h
@@ -43,9 +43,9 @@ struct FIOSPlatformProperties
 		return true;
 	}
     
-	static FORCEINLINE bool SupportsBuildTarget( EBuildTargets::Type BuildTarget )
+	static FORCEINLINE bool SupportsBuildTarget( EBuildTargetType TargetType )
 	{
-		return (BuildTarget == EBuildTargets::Game);
+		return (TargetType == EBuildTargetType::Game);
 	}
 
 	static FORCEINLINE bool SupportsLowQualityLightmaps()
