Add functionality to allow querying the available targets for a project from the editor, and use it to allow packaging client and server targets from the "Package Project" menu.

diff --git a/Engine/Source/Developer/DesktopPlatform/Private/DesktopPlatformBase.cpp b/Engine/Source/Developer/DesktopPlatform/Private/DesktopPlatformBase.cpp
--- a/Engine/Source/Developer/DesktopPlatform/Private/DesktopPlatformBase.cpp
+++ b/Engine/Source/Developer/DesktopPlatform/Private/DesktopPlatformBase.cpp
@@ -13,6 +13,8 @@
 #include "Serialization/JsonReader.h"
 #include "Serialization/JsonSerializer.h"
 #include "Modules/ModuleManager.h"
+#include "DesktopPlatformPrivate.h"
+#include "Misc/OutputDeviceRedirector.h"
 
 
 #define LOCTEXT_NAMESPACE "DesktopPlatform"
@@ -746,6 +748,153 @@ FProcHandle FDesktopPlatformBase::InvokeUnrealBuildToolAsync(const FString& InCm
 	return ProcHandle;
 }
 
+struct FTargetFileVisitor : IPlatformFile::FDirectoryStatVisitor
+{
+	TSet<FString>& RemainingTargetNames;
+	FDateTime MaxDateTime;
+	TArray<FString> SubDirectories;
+	bool bSearchSubDirectories;
+
+	FTargetFileVisitor(TSet<FString>& InRemainingTargetNames, FDateTime InMaxDateTime)
+		: RemainingTargetNames(InRemainingTargetNames)
+		, MaxDateTime(InMaxDateTime)
+		, bSearchSubDirectories(false)
+	{
+	}
+
+	virtual bool Visit(const TCHAR* FileNameOrDirectory, const FFileStatData& StatData) override
+	{
+		if (StatData.bIsDirectory)
+		{
+			SubDirectories.Add(FileNameOrDirectory);
+			return true;
+		}
+
+		int32 Length = FCString::Strlen(FileNameOrDirectory);
+
+		static const TCHAR TargetExt[] = TEXT(".Target.cs");
+		static const int32 TargetExtLen = ARRAY_COUNT(TargetExt) - 1;
+		if (Length > TargetExtLen && FCString::Stricmp(FileNameOrDirectory + Length - TargetExtLen, TargetExt) == 0)
+		{
+			FString TargetName = FPaths::GetCleanFilename(FString(Length - TargetExtLen, FileNameOrDirectory));
+			return (StatData.ModificationTime < MaxDateTime && RemainingTargetNames.Remove(TargetName) == 1);
+		}
+		
+		static const TCHAR ModuleExt[] = TEXT(".Build.cs");
+		static const int32 ModuleExtLen = ARRAY_COUNT(ModuleExt) - 1;
+		if (Length > ModuleExtLen && FCString::Stricmp(FileNameOrDirectory + Length - ModuleExtLen, ModuleExt) == 0)
+		{
+			bSearchSubDirectories = false;
+			return true;
+		}
+
+		return true;
+	}
+};
+
+bool IsTargetInfoValid(const TArray<FTargetInfo>& Targets, const FString& SourceDir, const FDateTime& LastModifiedTime)
+{
+	// Create the state 
+	TSet<FString> RemainingTargetNames;
+	for (const FTargetInfo& Target : Targets)
+	{
+		RemainingTargetNames.Add(Target.Name);
+	}
+
+	// Loop through all the directories
+	TArray<FString> DirectoryNames = { SourceDir };
+	for(int Idx = 0; Idx < DirectoryNames.Num(); Idx++)
+	{
+		FTargetFileVisitor Visitor(RemainingTargetNames, LastModifiedTime);
+		if(!IFileManager::Get().IterateDirectoryStat(*DirectoryNames[Idx], Visitor))
+		{
+			return false;
+		}
+		if(Visitor.bSearchSubDirectories)
+		{
+			DirectoryNames += Visitor.SubDirectories;
+		}
+	}
+
+	// If we found all the previous target files
+	return RemainingTargetNames.Num() == 0;
+}
+
+const TArray<FTargetInfo>& FDesktopPlatformBase::GetTargetsForProject(const FString& ProjectFile) const
+{
+	// Normalize the project filename
+	FString NormalizedProjectFile = ProjectFile;
+	FPaths::NormalizeFilename(NormalizedProjectFile);
+
+	// Check if there's already a cached entry for this project
+	const TArray<FTargetInfo>* Targets = ProjectFileToTargets.Find(NormalizedProjectFile);
+	if (Targets != nullptr)
+	{
+		return *Targets;
+	}
+
+	// Get the base directory for the project (or the engine directory)
+	FString ProjectDir;
+	if(ProjectFile.Len() == 0)
+	{
+		ProjectDir = FPaths::EngineDir();
+	}
+	else
+	{
+		ProjectDir = FPaths::GetPath(ProjectFile);
+	}
+
+	// Get the project source directory. If it doesn't exist, there are no targets for this project.
+	FString ProjectSourceDir = ProjectDir / TEXT("Source");
+	if (!IFileManager::Get().DirectoryExists(*ProjectSourceDir))
+	{
+		return ProjectFileToTargets.Add(NormalizedProjectFile, TArray<FTargetInfo>());
+	}
+
+	// Get the path to the info filename
+	FString InfoFileName = ProjectDir / TEXT("Intermediate/TargetInfo.json");
+
+	// Check if the file already exists
+	FFileStatData StatData = IFileManager::Get().GetStatData(*InfoFileName);
+	if(StatData.bIsValid)
+	{
+		// Read it in and check it's still valid
+		TArray<FTargetInfo> NewTargets;
+		if(ReadTargetInfo(InfoFileName, NewTargets) && IsTargetInfoValid(NewTargets, ProjectSourceDir, StatData.ModificationTime))
+		{
+			return ProjectFileToTargets.Emplace(MoveTemp(NormalizedProjectFile), MoveTemp(NewTargets));
+		}
+	}
+
+	// Otherwise, we'll have to run UBT to update it
+	FString Arguments = TEXT("-Mode=QueryTargets");
+	if(ProjectFile.Len() > 0)
+	{
+		Arguments += FString::Printf(TEXT(" -Project=\"%s\""), *IFileManager::Get().ConvertToAbsolutePathForExternalAppForRead(*ProjectFile));
+	}
+	Arguments += FString::Printf(TEXT(" -Output=\"%s\""), *IFileManager::Get().ConvertToAbsolutePathForExternalAppForRead(*InfoFileName));
+
+	// Run UBT to update the list of targets
+	FString Output;
+	int32 ReturnCode;
+	const_cast<FDesktopPlatformBase*>(this)->InvokeUnrealBuildToolSync(Arguments, *GLog, false, ReturnCode, Output);
+
+	// Try to read the new targets
+	TArray<FTargetInfo> NewTargets;
+	if(!ReadTargetInfo(InfoFileName, NewTargets))
+	{
+		UE_LOG(LogDesktopPlatform, Warning, TEXT("Unable to read target info for %s"), (ProjectFile.Len() == 0)? TEXT("engine") : *ProjectFile);
+	}
+
+	// Add it to the cache
+	return ProjectFileToTargets.Emplace(MoveTemp(NormalizedProjectFile), MoveTemp(NewTargets));
+}
+
+const TArray<FTargetInfo>& FDesktopPlatformBase::GetTargetsForCurrentProject() const
+{
+	return GetTargetsForProject(FPaths::GetProjectFilePath());
+}
+
 bool FDesktopPlatformBase::GetSolutionPath(FString& OutSolutionPath)
 {
 	// Get the platform-specific suffix for solution files
@@ -1167,7 +1316,7 @@ bool FDesktopPlatformBase::BuildUnrealBuildTool(const FString& RootDir, FOutputD
 
 	// Check the project file exists
 	FString CsProjLocation = GetUnrealBuildToolProjectFileName(RootDir);
-	if(!FPaths::FileExists(CsProjLocation))
+	if (!FPaths::FileExists(CsProjLocation))
 	{
 		Ar.Logf(TEXT("Project file not found at %s"), *CsProjLocation);
 		return false;
@@ -1221,7 +1370,7 @@ bool FDesktopPlatformBase::BuildUnrealBuildTool(const FString& RootDir, FOutputD
 
 	// If the executable appeared where we expect it, then we were successful
 	FString UnrealBuildToolExePath = GetUnrealBuildToolExecutableFilename(RootDir);
-	if(!FPaths::FileExists(UnrealBuildToolExePath))
+	if (!FPaths::FileExists(UnrealBuildToolExePath))
 	{
 		Ar.Logf(TEXT("Missing %s after build"), *UnrealBuildToolExePath);
 		return false;
@@ -1230,6 +1379,65 @@ bool FDesktopPlatformBase::BuildUnrealBuildTool(const FString& RootDir, FOutputD
 	return true;
 }
 
+bool FDesktopPlatformBase::ReadTargetInfo(const FString& FileName, TArray<FTargetInfo>& Targets)
+{
+	// Read the file to a string
+	FString Contents;
+	if (!FFileHelper::LoadFileToString(Contents, *FileName))
+	{
+		return false;
+	}
+
+	// Deserialize a JSON object from the string
+	TSharedPtr<FJsonObject> ObjectPtr;
+	TSharedRef<TJsonReader<>> Reader = TJsonReaderFactory<>::Create(Contents);
+	if (!FJsonSerializer::Deserialize(Reader, ObjectPtr) || !ObjectPtr.IsValid())
+	{
+		return false;
+	}
+
+	// Get the targets array
+	const TArray<TSharedPtr<FJsonValue>>* TargetArray;
+	if (!ObjectPtr->TryGetArrayField(TEXT("Targets"), TargetArray))
+	{
+		return false;
+	}
+
+	// Presize the output array
+	Targets.SetNum(TargetArray->Num());
+
+	// Parse the entries
+	for (int Idx = 0; Idx < TargetArray->Num(); Idx++)
+	{
+		const FJsonValue& TargetValue = *(*TargetArray)[Idx].Get();
+		if (TargetValue.Type != EJson::Object)
+		{
+			return false;
+		}
+
+		const FJsonObject& TargetObject = *TargetValue.AsObject();
+		if(!TargetObject.TryGetStringField(TEXT("Name"), Targets[Idx].Name) || !TargetObject.TryGetStringField(TEXT("Path"), Targets[Idx].Path))
+		{
+			return false;
+		}
+
+		FString Type;
+		if (!TargetObject.TryGetStringField(TEXT("Type"), Type))
+		{
+			return false;
+		}
+
+		LexFromString(Targets[Idx].Type, *Type);
+
+		if (Targets[Idx].Type == EBuildTargetType::Unknown)
+		{
+			return false;
+		}
+	}
+
+	return true;
+}
+
 FString FDesktopPlatformBase::GetUnrealBuildToolProjectFileName(const FString& RootDir) const
 {
 	return FPaths::ConvertRelativePathToFull(RootDir / TEXT("Engine/Source/Programs/UnrealBuildTool/UnrealBuildTool.csproj"));
diff --git a/Engine/Source/Developer/DesktopPlatform/Private/DesktopPlatformBase.h b/Engine/Source/Developer/DesktopPlatform/Private/DesktopPlatformBase.h
--- a/Engine/Source/Developer/DesktopPlatform/Private/DesktopPlatformBase.h
+++ b/Engine/Source/Developer/DesktopPlatform/Private/DesktopPlatformBase.h
@@ -49,6 +49,9 @@ public:
 	virtual bool InvokeUnrealBuildToolSync(const FString& InCmdLineParams, FOutputDevice &Ar, bool bSkipBuildUBT, int32& OutReturnCode, FString& OutProcOutput) override;
 	virtual FProcHandle InvokeUnrealBuildToolAsync(const FString& InCmdLineParams, FOutputDevice &Ar, void*& OutReadPipe, void*& OutWritePipe, bool bSkipBuildUBT = false) override;
 
+	virtual const TArray<FTargetInfo>& GetTargetsForProject(const FString& ProjectFile) const override;
+	virtual const TArray<FTargetInfo>& GetTargetsForCurrentProject() const override;
+
 	virtual bool GetSolutionPath(FString& OutSolutionPath) override;
 
 	virtual bool EnumerateProjectsKnownByEngine(const FString &Identifier, bool bIncludeNativeProjects, TArray<FString> &OutProjectFileNames) override;
@@ -59,6 +62,7 @@ private:
 	FDateTime LauncherInstallationTimestamp;
 	TMap<FString, FString> LauncherInstallationList;
 	TMap<FString, FUProjectDictionary> CachedProjectDictionaries;
+	mutable TMap<FString, TArray<FTargetInfo>> ProjectFileToTargets;
 
 	void ReadLauncherInstallationList();
 	void CheckForLauncherEngineInstallation(const FString &AppId, const FString &Identifier, TMap<FString, FString> &OutInstallations);
@@ -72,6 +76,8 @@ private:
 	void GetProjectBuildProducts(const FString& ProjectFileName, TArray<FString> &OutFileNames, TArray<FString> &OutDirectoryNames);
 	bool BuildUnrealBuildTool(const FString& RootDir, FOutputDevice &Ar);
 
+	static bool ReadTargetInfo(const FString& FileName, TArray<FTargetInfo>& Targets);
+
 protected:
 
 	FString GetUnrealBuildToolProjectFileName(const FString& RootDir) const;
diff --git a/Engine/Source/Developer/DesktopPlatform/Public/IDesktopPlatform.h b/Engine/Source/Developer/DesktopPlatform/Public/IDesktopPlatform.h
--- a/Engine/Source/Developer/DesktopPlatform/Public/IDesktopPlatform.h
+++ b/Engine/Source/Developer/DesktopPlatform/Public/IDesktopPlatform.h
@@ -58,7 +58,19 @@ public:
 	}
 };
 
+/**
+ * Information about a target supported by a project
+ */
+struct FTargetInfo
+{
+	FString Name;
+	FString Path;
+	EBuildTargetType Type;
+};
 
+/**
+ * Interface for functionality supported by desktop platforms
+ */
 class IDesktopPlatform
 {
 public:
@@ -387,6 +399,19 @@ public:
 	*/
 	virtual bool IsUnrealBuildToolRunning() = 0;
 
+	/**
+	 * Gets information about the build targets supported by a particular project.
+	 *
+	 * @return Array of TargetInfo objects
+	 */
+	virtual const TArray<FTargetInfo>& GetTargetsForProject(const FString& ProjectFile) const = 0;
+
+	/**
+	 * Gets information about the build targets supported by the current project.
+	 *
+	 * @return Array of TargetInfo objects
+	 */
+	virtual const TArray<FTargetInfo>& GetTargetsForCurrentProject() const = 0;
 
 	/**
 	* Gets the path to the solution for the current project
diff --git a/Engine/Source/Developer/TargetPlatform/Private/InstalledPlatformInfo.cpp b/Engine/Source/Developer/TargetPlatform/Private/InstalledPlatformInfo.cpp
--- a/Engine/Source/Developer/TargetPlatform/Private/InstalledPlatformInfo.cpp
+++ b/Engine/Source/Developer/TargetPlatform/Private/InstalledPlatformInfo.cpp
@@ -172,6 +172,44 @@ bool FInstalledPlatformInfo::IsValidTargetType(EBuildTargetType TargetType) cons
 	);
 }
 
+bool FInstalledPlatformInfo::IsValid(TOptional<EBuildTargetType> TargetType, TOptional<FString> Platform, TOptional<EBuildConfiguration> Configuration, EProjectType ProjectType, EInstalledPlatformState State) const
+{
+	if (!FApp::IsEngineInstalled())
+	{
+		return true;
+	}
+
+	for (const FInstalledPlatformConfiguration& Config : InstalledPlatformConfigurations)
+	{
+		// Check whether this configuration matches all the criteria
+		if(TargetType.IsSet() && Config.PlatformType != TargetType.GetValue())
+		{
+			continue;
+		}
+		if(Platform.IsSet() && Config.PlatformName != Platform.GetValue())
+		{
+			continue;
+		}
+		if(Configuration.IsSet() && Config.Configuration != Configuration.GetValue())
+		{
+			continue;
+		}
+		if(ProjectType != EProjectType::Any && Config.ProjectType != EProjectType::Any && Config.ProjectType != ProjectType)
+		{
+			continue;
+		}
+		if(State == EInstalledPlatformState::Downloaded && !Config.RequiredFile.IsEmpty() && !FPaths::FileExists(Config.RequiredFile))
+		{
+			continue;
+		}
+
+		// Success!
+		return true;
+	}
+
+	return false;
+}
+
 bool FInstalledPlatformInfo::IsValidPlatformArchitecture(const FString& PlatformName, const FString& Architecture) const
 {
 	return ContainsValidConfiguration(
diff --git a/Engine/Source/Developer/TargetPlatform/Public/InstalledPlatformInfo.h b/Engine/Source/Developer/TargetPlatform/Public/InstalledPlatformInfo.h
--- a/Engine/Source/Developer/TargetPlatform/Public/InstalledPlatformInfo.h
+++ b/Engine/Source/Developer/TargetPlatform/Public/InstalledPlatformInfo.h
@@ -15,6 +15,15 @@ enum class EProjectType : uint8
 
 EProjectType EProjectTypeFromString(const FString& ProjectTypeName);
 
+enum class EInstalledPlatformState
+{
+	/** Query whether the platform is supported */
+	Supported,
+
+	/** Query whether the platform has been downloaded */
+	Downloaded,
+};
+
 /**
  * Information about a single installed platform configuration
  */
@@ -82,6 +91,16 @@ public:
 	 */
 	bool IsValidTargetType(EBuildTargetType TargetType) const;
 
+	/** Determines whether the given target type is supported
+	 * @param TargetType The target type being built
+	 * @param Platform The platform being built
+	 * @param Configuration The configuration being built
+	 * @param ProjectType The project type required
+	 * @param State State of the given platform support
+	 * @return True if the target can be built
+	 */
+	bool IsValid(TOptional<EBuildTargetType> TargetType, TOptional<FString> Platform, TOptional<EBuildConfiguration> Configuration, EProjectType ProjectType, EInstalledPlatformState State) const;
+
 	/**
 	 * Queries whether a platform architecture is valid for any configuration
 	 * @param PlatformName Name of the platform's binary folder (eg. Win64, Android)
diff --git a/Engine/Source/Editor/MainFrame/Private/Frame/MainFrameActions.cpp b/Engine/Source/Editor/MainFrame/Private/Frame/MainFrameActions.cpp
--- a/Engine/Source/Editor/MainFrame/Private/Frame/MainFrameActions.cpp
+++ b/Engine/Source/Editor/MainFrame/Private/Frame/MainFrameActions.cpp
@@ -25,6 +25,7 @@
 #include "SourceCodeNavigation.h"
 #include "SourceControlWindows.h"
 #include "ISettingsModule.h"
+#include "Interfaces/IProjectManager.h"
 #include "Interfaces/ITargetPlatform.h"
 #include "Interfaces/ITargetPlatformManagerModule.h"
 #include "PlatformInfo.h"
@@ -733,12 +734,14 @@ void FMainFrameActionCallbacks::PackageProject( const FName InPlatformInfoName )
 		OptionalParams += FString::Printf(TEXT(" -NumCookersToSpawn=%d"), NumCookers); 
 	}
 
-	FString Configuration = StaticEnum<EProjectPackagingBuildConfigurations>()->GetNameStringByValue(PackagingSettings->BuildConfiguration);
-	Configuration = Configuration.Replace(TEXT("PPBC_"), TEXT(""));
-	if (Configuration.Right(6) == TEXT("Client"))
+	const UProjectPackagingSettings::FConfigurationInfo& Info = UProjectPackagingSettings::ConfigurationInfo[PackagingSettings->BuildConfiguration];
+	if (Info.TargetType == EBuildTargetType::Client)
 	{
 		OptionalParams += TEXT(" -client");
-		Configuration = Configuration.LeftChop(6);
+	}
+	else if (Info.TargetType == EBuildTargetType::Server)
+	{
+		OptionalParams += TEXT(" -server");
 	}
 
 	FString ProjectPath = FPaths::IsProjectFilePathSet() ? FPaths::ConvertRelativePathToFull(FPaths::GetProjectFilePath()) : FPaths::RootDir() / FApp::GetProjectName() / FApp::GetProjectName() + TEXT(".uproject");
@@ -748,7 +751,7 @@ void FMainFrameActionCallbacks::PackageProject( const FName InPlatformInfoName )
 		FApp::IsEngineInstalled() ? TEXT(" -installed") : TEXT(""),
 		*ProjectPath,
 		*PackagingSettings->StagingDirectory.Path,
-		*Configuration,
+		LexToString(Info.Configuration),
 		*FUnrealEdMisc::Get().GetExecutableForCommandlets(),
 		*OptionalParams
 	);
diff --git a/Engine/Source/Editor/MainFrame/Private/Menus/PackageProjectMenu.cpp b/Engine/Source/Editor/MainFrame/Private/Menus/PackageProjectMenu.cpp
new file mode 100644
--- /dev/null
+++ b/Engine/Source/Editor/MainFrame/Private/Menus/PackageProjectMenu.cpp
@@ -0,0 +1,10 @@
+// Copyright 1998-2019 Epic Games, Inc. All Rights Reserved.
+
+#pragma once
+
+#include "PackageProjectMenu.h"
+
+#define LOCTEXT_NAMESPACE "FPackageProjectMenu"
+
+
+#undef LOCTEXT_NAMESPACE
diff --git a/Engine/Source/Editor/MainFrame/Private/Menus/PackageProjectMenu.h b/Engine/Source/Editor/MainFrame/Private/Menus/PackageProjectMenu.h
--- a/Engine/Source/Editor/MainFrame/Private/Menus/PackageProjectMenu.h
+++ b/Engine/Source/Editor/MainFrame/Private/Menus/PackageProjectMenu.h
@@ -16,6 +16,7 @@
 #include "Interfaces/IProjectManager.h"
 #include "InstalledPlatformInfo.h"
 #include "Misc/DataDrivenPlatformInfoRegistry.h"
+#include "DesktopPlatformModule.h"
 
 #define LOCTEXT_NAMESPACE "FPackageProjectMenu"
 
@@ -168,107 +169,61 @@ protected:
 	 */
 	static void MakeBuildConfigurationsMenu(FMenuBuilder& MenuBuilder)
 	{
-		// Only show the debug game configurations if the game has source code. 
-		TArray<FString> TargetFileNames;
-		IFileManager::Get().FindFiles(TargetFileNames, *(FPaths::GameSourceDir() / TEXT("*.target.cs")), true, false);
-		
-		if (TargetFileNames.Num() > 0)
-		{
-			MenuBuilder.AddMenuEntry(
-				LOCTEXT("DebugConfiguration", "DebugGame"),
-				LOCTEXT("DebugConfigurationTooltip", "Package the project for debugging"),
-				FSlateIcon(),
-				FUIAction(
-					FExecuteAction::CreateStatic(&FMainFrameActionCallbacks::PackageBuildConfiguration, PPBC_DebugGame),
-					FCanExecuteAction::CreateStatic(&FMainFrameActionCallbacks::CanPackageBuildConfiguration, PPBC_DebugGame),
-					FIsActionChecked::CreateStatic(&FMainFrameActionCallbacks::PackageBuildConfigurationIsChecked, PPBC_DebugGame)
-				),
-				NAME_None,
-				EUserInterfaceActionType::RadioButton
-			);
-		}
+		// Check if the project has code
+		FProjectStatus ProjectStatus;
+		bool bHasCode = IProjectManager::Get().QueryStatusForCurrentProject(ProjectStatus) && ProjectStatus.bCodeBasedProject;
 
-		// Add the Client configurations if there is a {ProjectName}Client.Target.cs file.
-		TArray<FString> ClientTargetFileNames;
-		if (FInstalledPlatformInfo::Get().IsValidTargetType(EBuildTargetType::Client))
+		// If if does, find all the targets
+		const TArray<FTargetInfo>* Targets = nullptr;
+		if (bHasCode)
 		{
-			IFileManager::Get().FindFiles(ClientTargetFileNames, *(FPaths::GameSourceDir() / TEXT("*client.target.cs")), true, false);
+			Targets = &(FDesktopPlatformModule::Get()->GetTargetsForCurrentProject());
 		}
 
-		if (ClientTargetFileNames.Num() > 0)
+		// Set up all the configurations
+		for (int32 Idx = 0; Idx < PPBC_Max; Idx++)
 		{
-			MenuBuilder.AddMenuEntry(
-				LOCTEXT("DebugClientConfiguration", "DebugGame Client"),
-				LOCTEXT("DebugClientConfigurationTooltip", "Package the project for debugging as a client"),
-				FSlateIcon(),
-				FUIAction(
-					FExecuteAction::CreateStatic(&FMainFrameActionCallbacks::PackageBuildConfiguration, PPBC_DebugGameClient),
-					FCanExecuteAction::CreateStatic(&FMainFrameActionCallbacks::CanPackageBuildConfiguration, PPBC_DebugGameClient),
-					FIsActionChecked::CreateStatic(&FMainFrameActionCallbacks::PackageBuildConfigurationIsChecked, PPBC_DebugGameClient)
-				),
-				NAME_None,
-				EUserInterfaceActionType::RadioButton
-			);
-		}
-
-		MenuBuilder.AddMenuEntry(
-			LOCTEXT("DevelopmentConfiguration", "Development"),
-			LOCTEXT("DevelopmentConfigurationTooltip", "Package the project for development"),
-			FSlateIcon(),
-			FUIAction(
-				FExecuteAction::CreateStatic(&FMainFrameActionCallbacks::PackageBuildConfiguration, PPBC_Development),
-				FCanExecuteAction::CreateStatic(&FMainFrameActionCallbacks::CanPackageBuildConfiguration, PPBC_Development),
-				FIsActionChecked::CreateStatic(&FMainFrameActionCallbacks::PackageBuildConfigurationIsChecked, PPBC_Development)
-			),
-			NAME_None,
-			EUserInterfaceActionType::RadioButton
-		);
+			EProjectPackagingBuildConfigurations PackagingConfiguration = (EProjectPackagingBuildConfigurations)Idx;
+			const UProjectPackagingSettings::FConfigurationInfo& Info = UProjectPackagingSettings::ConfigurationInfo[Idx];
 
-		if (ClientTargetFileNames.Num() > 0)
-		{
-			MenuBuilder.AddMenuEntry(
-				LOCTEXT("DevelopmentClientConfiguration", "Development Client"),
-				LOCTEXT("DevelopmentClientConfigurationTooltip", "Package the project for development as a client"),
-				FSlateIcon(),
-				FUIAction(
-					FExecuteAction::CreateStatic(&FMainFrameActionCallbacks::PackageBuildConfiguration, PPBC_DevelopmentClient),
-					FCanExecuteAction::CreateStatic(&FMainFrameActionCallbacks::CanPackageBuildConfiguration, PPBC_DevelopmentClient),
-					FIsActionChecked::CreateStatic(&FMainFrameActionCallbacks::PackageBuildConfigurationIsChecked, PPBC_DevelopmentClient)
-				),
-				NAME_None,
-				EUserInterfaceActionType::RadioButton
-			);
-		}
+			// Make sure the selected configuration is supported
+			EProjectType ProjectType = bHasCode? EProjectType::Code : EProjectType::Content;
+			if(!FInstalledPlatformInfo::Get().IsValid(Info.TargetType, TOptional<FString>(), Info.Configuration, ProjectType, EInstalledPlatformState::Downloaded))
+			{
+				continue;
+			}
 
-		MenuBuilder.AddMenuEntry(
-			LOCTEXT("ShippingConfiguration", "Shipping"),
-			LOCTEXT("ShippingConfigurationTooltip", "Package the project for shipping"),
-			FSlateIcon(),
-			FUIAction(
-				FExecuteAction::CreateStatic(&FMainFrameActionCallbacks::PackageBuildConfiguration, PPBC_Shipping),
-				FCanExecuteAction(),
-				FIsActionChecked::CreateStatic(&FMainFrameActionCallbacks::PackageBuildConfigurationIsChecked, PPBC_Shipping)
-			),
-			NAME_None,
-			EUserInterfaceActionType::RadioButton
-		);
+			// Check the target type is valid
+			if(bHasCode)
+			{
+				EBuildTargetType TargetType = Info.TargetType;
+				if(!Targets->ContainsByPredicate([TargetType](const FTargetInfo& Target) -> bool { return Target.Type == TargetType; }))
+				{
+					continue;
+				}
+			}
+			else
+			{
+				if(Info.Configuration == EBuildConfiguration::DebugGame)
+				{
+					continue;
+				}
+			}
 
-		if (ClientTargetFileNames.Num() > 0)
-		{
+			// Add the menu item
 			MenuBuilder.AddMenuEntry(
-				LOCTEXT("ShippingClientConfiguration", "Shipping Client"),
-				LOCTEXT("ShippingClientConfigurationTooltip", "Package the project for shipping as a client"),
+				Info.Name, 
+				Info.ToolTip,
 				FSlateIcon(),
 				FUIAction(
-					FExecuteAction::CreateStatic(&FMainFrameActionCallbacks::PackageBuildConfiguration, PPBC_ShippingClient),
-					FCanExecuteAction(),
-					FIsActionChecked::CreateStatic(&FMainFrameActionCallbacks::PackageBuildConfigurationIsChecked, PPBC_ShippingClient)
+					FExecuteAction::CreateStatic(&FMainFrameActionCallbacks::PackageBuildConfiguration, PackagingConfiguration),
+					FCanExecuteAction::CreateStatic(&FMainFrameActionCallbacks::CanPackageBuildConfiguration, PackagingConfiguration),
+					FIsActionChecked::CreateStatic(&FMainFrameActionCallbacks::PackageBuildConfigurationIsChecked, PackagingConfiguration)
 				),
 				NAME_None,
 				EUserInterfaceActionType::RadioButton
 			);
 		}
-
 	}
 };
 
diff --git a/Engine/Source/Editor/UnrealEd/Classes/Settings/ProjectPackagingSettings.h b/Engine/Source/Editor/UnrealEd/Classes/Settings/ProjectPackagingSettings.h
--- a/Engine/Source/Editor/UnrealEd/Classes/Settings/ProjectPackagingSettings.h
+++ b/Engine/Source/Editor/UnrealEd/Classes/Settings/ProjectPackagingSettings.h
@@ -15,22 +15,52 @@ UENUM()
 enum EProjectPackagingBuildConfigurations
 {
 	/** Debug configuration. */
-	PPBC_DebugGame UMETA(DisplayName="DebugGame"),
+	PPBC_Debug UMETA(DisplayName="Debug"),
 
 	/** Debug Client configuration. */
+	PPBC_DebugClient UMETA(DisplayName = "Debug Client"),
+
+	/** Debug Server configuration. */
+	PPBC_DebugServer UMETA(DisplayName = "Debug Server"),
+
+	/** DebugGame configuration. */
+	PPBC_DebugGame UMETA(DisplayName="DebugGame"),
+
+	/** DebugGame Client configuration. */
 	PPBC_DebugGameClient UMETA(DisplayName = "DebugGame Client"),
 
+	/** DebugGame Server configuration. */
+	PPBC_DebugGameServer UMETA(DisplayName = "DebugGame Server"),
+
 	/** Development configuration. */
 	PPBC_Development UMETA(DisplayName="Development"),
 
 	/** Development Client configuration. */
 	PPBC_DevelopmentClient UMETA(DisplayName = "Development Client"),
 
+	/** Development Server configuration. */
+	PPBC_DevelopmentServer UMETA(DisplayName = "Development Server"),
+
+	/** Test configuration. */
+	PPBC_Test UMETA(DisplayName="Test"),
+
+	/** Test Client configuration. */
+	PPBC_TestClient UMETA(DisplayName = "Test Client"),
+
+	/** Test Server configuration. */
+	PPBC_TestServer UMETA(DisplayName = "Test Server"),
+
 	/** Shipping configuration. */
 	PPBC_Shipping UMETA(DisplayName="Shipping"),
 
 	/** Shipping Client configuration. */
-	PPBC_ShippingClient UMETA(DisplayName = "Shipping Client")
+	PPBC_ShippingClient UMETA(DisplayName = "Shipping Client"),
+
+	/** Shipping Server configuration. */
+	PPBC_ShippingServer UMETA(DisplayName = "Shipping Server"),
+
+	/** Number of entries in the enum. */
+	PPBC_Max
 };
 
 /**
@@ -100,6 +130,21 @@ class UNREALED_API UProjectPackagingSettings
 	GENERATED_UCLASS_BODY()
 
 public:
+	/**
+	 * Information about each packaging configuration
+	 */
+	struct FConfigurationInfo
+	{
+		EBuildConfiguration Configuration;
+		EBuildTargetType TargetType;
+		FText Name;
+		FText ToolTip;
+	};
+
+	/**
+	 * Static array of information about each configuration
+	 */
+	static const FConfigurationInfo ConfigurationInfo[PPBC_Max];
 
 	/** Specifies whether to build the game executable during packaging. */
 	UPROPERTY(config, EditAnywhere, Category=Project)
diff --git a/Engine/Source/Editor/UnrealEd/Private/Settings/SettingsClasses.cpp b/Engine/Source/Editor/UnrealEd/Private/Settings/SettingsClasses.cpp
--- a/Engine/Source/Editor/UnrealEd/Private/Settings/SettingsClasses.cpp
+++ b/Engine/Source/Editor/UnrealEd/Private/Settings/SettingsClasses.cpp
@@ -781,6 +781,29 @@ void ULevelEditorViewportSettings::PostEditChangeProperty( struct FPropertyChang
 /* UProjectPackagingSettings interface
  *****************************************************************************/
 
+const UProjectPackagingSettings::FConfigurationInfo UProjectPackagingSettings::ConfigurationInfo[PPBC_Max] = 
+{
+	/* PPBC_Debug */             { EBuildConfiguration::Debug, EBuildTargetType::Game, LOCTEXT("DebugConfiguration", "Debug"), LOCTEXT("DebugConfigurationTooltip", "Package the game in Debug configuration") },
+	/* PPBC_DebugClient */       { EBuildConfiguration::Debug, EBuildTargetType::Client, LOCTEXT("DebugClientConfiguration", "Debug Client"), LOCTEXT("DebugClientConfigurationTooltip", "Package the client in Debug configuration") },
+	/* PPBC_DebugServer */       { EBuildConfiguration::Debug, EBuildTargetType::Server, LOCTEXT("DebugServerConfiguration", "Debug Server"), LOCTEXT("DebugServerConfigurationTooltip", "Package the server in Debug configuration") },
+
+	/* PPBC_DebugGame */         { EBuildConfiguration::DebugGame, EBuildTargetType::Game, LOCTEXT("DebugGameConfiguration", "DebugGame"), LOCTEXT("DebugGameConfigurationTooltip", "Package the game in DebugGame configuration") },
+	/* PPBC_DebugGameClient */   { EBuildConfiguration::DebugGame, EBuildTargetType::Client, LOCTEXT("DebugGameClientConfiguration", "DebugGame Client"), LOCTEXT("DebugGameClientConfigurationTooltip", "Package the client in DebugGame configuration") },
+	/* PPBC_DebugGameServer */   { EBuildConfiguration::DebugGame, EBuildTargetType::Server, LOCTEXT("DebugGameServerConfiguration", "DebugGame Server"), LOCTEXT("DebugGameServerConfigurationTooltip", "Package the server in DebugGame configuration") },
+
+	/* PPBC_Development */       { EBuildConfiguration::Development, EBuildTargetType::Game, LOCTEXT("DevelopmentConfiguration", "Development"), LOCTEXT("DevelopmentConfigurationTooltip", "Package the game in Development configuration") },
+	/* PPBC_DevelopmentClient */ { EBuildConfiguration::Development, EBuildTargetType::Client, LOCTEXT("DevelopmentClientConfiguration", "Development Client"), LOCTEXT("DevelopmentClientConfigurationTooltip", "Package the client in Development configuration") },
+	/* PPBC_DevelopmentServer */ { EBuildConfiguration::Development, EBuildTargetType::Server, LOCTEXT("DevelopmentServerConfiguration", "Development Server"), LOCTEXT("DevelopmentServerConfigurationTooltip", "Package the server in Development configuration") },
+
+	/* PPBC_Test */              { EBuildConfiguration::Test, EBuildTargetType::Game, LOCTEXT("TestConfiguration", "Test"), LOCTEXT("TestConfigurationTooltip", "Package the game in Test configuration") },
+	/* PPBC_TestClient */        { EBuildConfiguration::Test, EBuildTargetType::Client, LOCTEXT("TestClientConfiguration", "Test Client"), LOCTEXT("TestClientConfigurationTooltip", "Package the client in Test configuration") },
+	/* PPBC_TestServer */        { EBuildConfiguration::Test, EBuildTargetType::Server, LOCTEXT("TestServerConfiguration", "Test Server"), LOCTEXT("TestServerConfigurationTooltip", "Package the server in Test configuration") },
+
+	/* PPBC_Shipping */          { EBuildConfiguration::Shipping, EBuildTargetType::Game, LOCTEXT("ShippingConfiguration", "Shipping"), LOCTEXT("ShippingConfigurationTooltip", "Package the game in Shipping configuration") },
+	/* PPBC_ShippingClient */    { EBuildConfiguration::Shipping, EBuildTargetType::Client, LOCTEXT("ShippingClientConfiguration", "Shipping Client"), LOCTEXT("ShippingClientConfigurationTooltip", "Package the client in Shipping configuration") },
+	/* PPBC_ShippingServer */    { EBuildConfiguration::Shipping, EBuildTargetType::Server, LOCTEXT("ShippingServerConfiguration", "Shipping Server"), LOCTEXT("ShippingServerConfigurationTooltip", "Package the server in Shipping configuration") }
+};
+
 UProjectPackagingSettings::UProjectPackagingSettings( const FObjectInitializer& ObjectInitializer )
 	: Super(ObjectInitializer)
 {
diff --git a/Engine/Source/Editor/UnrealEd/Private/UnrealEdGlobals.cpp b/Engine/Source/Editor/UnrealEd/Private/UnrealEdGlobals.cpp
--- a/Engine/Source/Editor/UnrealEd/Private/UnrealEdGlobals.cpp
+++ b/Engine/Source/Editor/UnrealEd/Private/UnrealEdGlobals.cpp
@@ -22,6 +22,7 @@
 #include "UnrealEngine.h"
 #include "UnrealEdMisc.h"
 #include "EditorModes.h"
+#include "DesktopPlatformModule.h"
 
 #include "DebugToolExec.h"
 #include "Modules/ModuleManager.h"
@@ -102,6 +103,9 @@ int32 EditorInit( IEngineLoop& EngineLoop )
 	// Set up the actor folders singleton
 	FActorFolders::Init();
 
+	// Cache the available targets for the current project, so we can display the appropriate options in the package project menu
+	FDesktopPlatformModule::Get()->GetTargetsForCurrentProject();
+
 	// =================== CORE EDITOR INIT FINISHED ===================
 
 	// Hide the splash screen now that everything is ready to go
diff --git a/Engine/Source/Programs/UnrealBuildTool/Configuration/UEBuildPlatform.cs b/Engine/Source/Programs/UnrealBuildTool/Configuration/UEBuildPlatform.cs
--- a/Engine/Source/Programs/UnrealBuildTool/Configuration/UEBuildPlatform.cs
+++ b/Engine/Source/Programs/UnrealBuildTool/Configuration/UEBuildPlatform.cs
@@ -51,7 +51,8 @@ namespace UnrealBuildTool
 		/// Finds all the UEBuildPlatformFactory types in this assembly and uses them to register all the available platforms
 		/// </summary>
 		/// <param name="bIncludeNonInstalledPlatforms">Whether to register platforms that are not installed</param>
-		public static void RegisterPlatforms(bool bIncludeNonInstalledPlatforms)
+		/// <param name="bHostPlatformOnly">Only register the host platform</param>
+		public static void RegisterPlatforms(bool bIncludeNonInstalledPlatforms, bool bHostPlatformOnly)
 		{
 			// Initialize the installed platform info
 			using(Timeline.ScopeEvent("Initializing InstalledPlatformInfo"))
@@ -78,6 +79,12 @@ namespace UnrealBuildTool
 						{
 							UEBuildPlatformFactory TempInst = (UEBuildPlatformFactory)Activator.CreateInstance(CheckType);
 							
+							if(bHostPlatformOnly && TempInst.TargetPlatform != BuildHostPlatform.Current.Platform)
+							{
+								continue;
+							}
+
+
 							// We need all platforms to be registered when we run -validateplatform command to check SDK status of each
 							if (bIncludeNonInstalledPlatforms || InstalledPlatformInfo.IsValidPlatform(TempInst.TargetPlatform))
 							{
diff --git a/Engine/Source/Programs/UnrealBuildTool/Modes/QueryTargetsMode.cs b/Engine/Source/Programs/UnrealBuildTool/Modes/QueryTargetsMode.cs
new file mode 100644
--- /dev/null
+++ b/Engine/Source/Programs/UnrealBuildTool/Modes/QueryTargetsMode.cs
@@ -0,0 +1,128 @@
+using System;
+using System.Collections.Generic;
+using System.Linq;
+using System.Text;
+using System.Threading.Tasks;
+using Tools.DotNETCommon;
+
+namespace UnrealBuildTool
+{
+	/// <summary>
+	/// Queries information about the targets supported by a project
+	/// </summary>
+	[ToolMode("QueryTargets", ToolModeOptions.XmlConfig | ToolModeOptions.BuildPlatformsHostOnly | ToolModeOptions.SingleInstance)]
+	class QueryTargetsMode : ToolMode
+	{
+		/// <summary>
+		/// Path to the project file to query
+		/// </summary>
+		[CommandLine("-Project=")]
+		FileReference ProjectFile = null;
+
+		/// <summary>
+		/// Path to the output file to receive information about the targets
+		/// </summary>
+		[CommandLine("-Output=")]
+		FileReference OutputFile = null;
+
+		/// <summary>
+		/// Execute the mode
+		/// </summary>
+		/// <param name="Arguments">Command line arguments</param>
+		/// <returns></returns>
+		public override int Execute(CommandLineArguments Arguments)
+		{
+			Arguments.ApplyTo(this);
+
+			// Create the build configuration object, and read the settings
+			BuildConfiguration BuildConfiguration = new BuildConfiguration();
+			XmlConfig.ApplyTo(BuildConfiguration);
+			Arguments.ApplyTo(BuildConfiguration);
+
+			// Ensure the path to the output file is valid
+			if(OutputFile == null)
+			{
+				OutputFile = GetDefaultOutputFile(ProjectFile);
+			}
+
+			// Create the rules assembly
+			RulesAssembly Assembly;
+			if(ProjectFile == null)
+			{
+				Assembly = RulesCompiler.CreateEnterpriseRulesAssembly(BuildConfiguration.bUsePrecompiled, BuildConfiguration.bSkipRulesCompile);
+			}
+			else
+			{
+				Assembly = RulesCompiler.CreateProjectRulesAssembly(ProjectFile, BuildConfiguration.bUsePrecompiled, BuildConfiguration.bSkipRulesCompile);
+			}
+
+			// Write information about these targets
+			WriteTargetInfo(ProjectFile, Assembly, OutputFile);
+			Log.TraceInformation("Written {0}", OutputFile);
+			return 0;
+		}
+
+		/// <summary>
+		/// Gets the path to the target info output file
+		/// </summary>
+		/// <param name="ProjectFile">Project file being queried</param>
+		/// <returns>Path to the output file</returns>
+		public static FileReference GetDefaultOutputFile(FileReference ProjectFile)
+		{
+			if (ProjectFile == null)
+			{
+				return FileReference.Combine(UnrealBuildTool.EngineDirectory, "Intermediate", "TargetInfo.json");
+			}
+			else
+			{
+				return FileReference.Combine(ProjectFile.Directory, "Intermediate", "TargetInfo.json");
+			}
+		}
+
+		/// <summary>
+		/// Writes information about the targets in an assembly to a file
+		/// </summary>
+		/// <param name="ProjectFile">The project file for the targets being built</param>
+		/// <param name="Assembly">The rules assembly for this target</param>
+		/// <param name="OutputFile">Output file to write to</param>
+		public static void WriteTargetInfo(FileReference ProjectFile, RulesAssembly Assembly, FileReference OutputFile)
+		{
+			// Construct all the targets in this assembly
+			List<string> TargetNames = new List<string>();
+			Assembly.GetAllTargetNames(TargetNames, false);
+
+			// Write the output file
+			DirectoryReference.CreateDirectory(OutputFile.Directory);
+			using (JsonWriter Writer = new JsonWriter(OutputFile))
+			{
+				Writer.WriteObjectStart();
+				Writer.WriteArrayStart("Targets");
+				foreach (string TargetName in TargetNames)
+				{
+					// Construct the rules object
+					TargetRules TargetRules;
+					try
+					{
+						string Architecture = UEBuildPlatform.GetBuildPlatform(BuildHostPlatform.Current.Platform).GetDefaultArchitecture(ProjectFile);
+						TargetRules = Assembly.CreateTargetRules(TargetName, BuildHostPlatform.Current.Platform, UnrealTargetConfiguration.Development, Architecture, ProjectFile, new CommandLineArguments(new string[0]));
+					}
+					catch (Exception Ex)
+					{
+						Log.TraceWarning("Unable to construct target rules for {0}", TargetName);
+						Log.TraceVerbose(ExceptionUtils.FormatException(Ex));
+						continue;
+					}
+
+					// Write the target info
+					Writer.WriteObjectStart();
+					Writer.WriteValue("Name", TargetName);
+					Writer.WriteValue("Path", Assembly.GetTargetFileName(TargetName).ToString());
+					Writer.WriteValue("Type", TargetRules.Type.ToString());
+					Writer.WriteObjectEnd();
+				}
+				Writer.WriteArrayEnd();
+				Writer.WriteObjectEnd();
+			}
+		}
+	}
+}
diff --git a/Engine/Source/Programs/UnrealBuildTool/ProjectFiles/ProjectFileGenerator.cs b/Engine/Source/Programs/UnrealBuildTool/ProjectFiles/ProjectFileGenerator.cs
--- a/Engine/Source/Programs/UnrealBuildTool/ProjectFiles/ProjectFileGenerator.cs
+++ b/Engine/Source/Programs/UnrealBuildTool/ProjectFiles/ProjectFileGenerator.cs
@@ -972,10 +972,17 @@ namespace UnrealBuildTool
 				// Generate IntelliSense data if we need to.  This involves having UBT simulate the action compilation of
 				// the targets so that we can extra the compiler defines, include paths, etc.
 				GenerateIntelliSenseData(Arguments, IntelliSenseTargetFiles);
-
+				
 				// Write the project files
 				WriteProjectFiles(PlatformProjectGenerators);
 				Log.TraceVerbose( "Project generation complete ({0} generated, {1} imported)", GeneratedProjectFiles.Count, OtherProjectFiles.Count );
+
+				// Generate all the target info files for the editor
+				foreach(FileReference ProjectFile in AllGameProjects)
+				{
+					RulesAssembly RulesAssembly = RulesCompiler.CreateProjectRulesAssembly(ProjectFile, false, false);
+					QueryTargetsMode.WriteTargetInfo(ProjectFile, RulesAssembly, QueryTargetsMode.GetDefaultOutputFile(ProjectFile));
+				}
 			}
 
 			return bSuccess;
diff --git a/Engine/Source/Programs/UnrealBuildTool/System/PlatformExports.cs b/Engine/Source/Programs/UnrealBuildTool/System/PlatformExports.cs
--- a/Engine/Source/Programs/UnrealBuildTool/System/PlatformExports.cs
+++ b/Engine/Source/Programs/UnrealBuildTool/System/PlatformExports.cs
@@ -156,7 +156,7 @@ namespace UnrealBuildTool
 			XmlConfig.ReadConfigFiles(null);
 
 			// Register all the platform classes
-			UEBuildPlatform.RegisterPlatforms(false);
+			UEBuildPlatform.RegisterPlatforms(false, false);
 			return true;
 		}
 	}
diff --git a/Engine/Source/Programs/UnrealBuildTool/System/RulesAssembly.cs b/Engine/Source/Programs/UnrealBuildTool/System/RulesAssembly.cs
--- a/Engine/Source/Programs/UnrealBuildTool/System/RulesAssembly.cs
+++ b/Engine/Source/Programs/UnrealBuildTool/System/RulesAssembly.cs
@@ -227,6 +227,20 @@ namespace UnrealBuildTool
 			ModuleNames.AddRange(CompiledAssembly.GetTypes().Where(x => x.IsClass && x.IsSubclassOf(typeof(ModuleRules)) && ModuleNameToModuleFile.ContainsKey(x.Name)).Select(x => x.Name));
 		}
 
+		/// <summary>
+		/// Fills a list with all the target names in this assembly
+		/// </summary>
+		/// <param name="TargetNames">List to receive the target names</param>
+		/// <param name="bIncludeParentAssembly">Whether to include targets in the parent assembly</param>
+		public void GetAllTargetNames(List<string> TargetNames, bool bIncludeParentAssembly)
+		{
+			if(Parent != null && bIncludeParentAssembly)
+			{
+				Parent.GetAllTargetNames(TargetNames, true);
+			}
+			TargetNames.AddRange(TargetNameToTargetFile.Keys);
+		}
+
 		/// <summary>
 		/// Tries to get the filename that declared the given type
 		/// </summary>
diff --git a/Engine/Source/Programs/UnrealBuildTool/System/ToolMode.cs b/Engine/Source/Programs/UnrealBuildTool/System/ToolMode.cs
--- a/Engine/Source/Programs/UnrealBuildTool/System/ToolMode.cs
+++ b/Engine/Source/Programs/UnrealBuildTool/System/ToolMode.cs
@@ -35,20 +35,25 @@ namespace UnrealBuildTool
 		/// </summary>
 		BuildPlatforms = 4,
 
+		/// <summary>
+		/// Registers build platforms
+		/// </summary>
+		BuildPlatformsHostOnly = 8,
+
 		/// <summary>
 		/// Registers build platforms for validation
 		/// </summary>
-		BuildPlatformsForValidation = 8,
+		BuildPlatformsForValidation = 16,
 
 		/// <summary>
 		/// Only allow a single instance running in the branch at once
 		/// </summary>
-		SingleInstance = 16,
+		SingleInstance = 32,
 
 		/// <summary>
 		/// Print out the total time taken to execute
 		/// </summary>
-		ShowExecutionTime = 32,
+		ShowExecutionTime = 64,
 	}
 
 	/// <summary>
diff --git a/Engine/Source/Programs/UnrealBuildTool/UnrealBuildTool.cs b/Engine/Source/Programs/UnrealBuildTool/UnrealBuildTool.cs
--- a/Engine/Source/Programs/UnrealBuildTool/UnrealBuildTool.cs
+++ b/Engine/Source/Programs/UnrealBuildTool/UnrealBuildTool.cs
@@ -495,14 +495,21 @@ namespace UnrealBuildTool
 				{
 					using(Timeline.ScopeEvent("UEBuildPlatform.RegisterPlatforms()"))
 					{
-						UEBuildPlatform.RegisterPlatforms(false);
+						UEBuildPlatform.RegisterPlatforms(false, false);
 					}
 				}
-				if((ModeOptions & ToolModeOptions.BuildPlatformsForValidation) != 0)
+				if ((ModeOptions & ToolModeOptions.BuildPlatformsHostOnly) != 0)
+				{
+					using (Timeline.ScopeEvent("UEBuildPlatform.RegisterPlatforms()"))
+					{
+						UEBuildPlatform.RegisterPlatforms(false, true);
+					}
+				}
+				if ((ModeOptions & ToolModeOptions.BuildPlatformsForValidation) != 0)
 				{
 					using(Timeline.ScopeEvent("UEBuildPlatform.RegisterPlatforms()"))
 					{
-						UEBuildPlatform.RegisterPlatforms(true);
+						UEBuildPlatform.RegisterPlatforms(true, false);
 					}
 				}
 
diff --git a/Engine/Source/Programs/UnrealBuildTool/UnrealBuildTool.csproj b/Engine/Source/Programs/UnrealBuildTool/UnrealBuildTool.csproj
--- a/Engine/Source/Programs/UnrealBuildTool/UnrealBuildTool.csproj
+++ b/Engine/Source/Programs/UnrealBuildTool/UnrealBuildTool.csproj
@@ -144,6 +144,7 @@
     <Compile Include="Modes\BuildMode.cs" />
     <Compile Include="Modes\CleanMode.cs" />
     <Compile Include="Modes\DeployMode.cs" />
+    <Compile Include="Modes\QueryTargetsMode.cs" />
     <Compile Include="Modes\ExecuteMode.cs" />
     <Compile Include="Modes\JsonExportMode.cs" />
     <Compile Include="Modes\GenerateProjectFilesMode.cs" />
@@ -283,44 +284,44 @@
     <Compile Include="Platform\Windows\PVSToolChain.cs" />
   </ItemGroup>
   <Choose>
-	<When Condition="Exists('Platform\HoloLens\UEBuildHoloLens.cs')">
-	  <ItemGroup>
-		<Compile Include="Platform\HoloLens\UEBuildHoloLens.cs" />
-		<Compile Include="Platform\HoloLens\HoloLensDeploy.cs" />
-		<Compile Include="Platform\HoloLens\HoloLensExports.cs" />
-		<Compile Include="Platform\HoloLens\HoloLensProjectGenerator.cs" />
-		<Compile Include="Platform\HoloLens\HoloLensToolChain.cs" />
-		<Compile Include="Platform\HoloLens\ManifestGenerator.cs" />
-		<Compile Include="Platform\HoloLens\WinMDRegistrationInfo.cs" />
-	  </ItemGroup>
-	</When>
+    <When Condition="Exists('Platform\HoloLens\UEBuildHoloLens.cs')">
+      <ItemGroup>
+        <Compile Include="Platform\HoloLens\UEBuildHoloLens.cs" />
+        <Compile Include="Platform\HoloLens\HoloLensDeploy.cs" />
+        <Compile Include="Platform\HoloLens\HoloLensExports.cs" />
+        <Compile Include="Platform\HoloLens\HoloLensProjectGenerator.cs" />
+        <Compile Include="Platform\HoloLens\HoloLensToolChain.cs" />
+        <Compile Include="Platform\HoloLens\ManifestGenerator.cs" />
+        <Compile Include="Platform\HoloLens\WinMDRegistrationInfo.cs" />
+      </ItemGroup>
+    </When>
   </Choose>
   <Choose>
     <When Condition="'$(OS)' != 'Windows_NT'">
-	  <PropertyGroup>
-		<DefineConstants>__MonoCS__</DefineConstants>
-	  </PropertyGroup>
-	</When>
+      <PropertyGroup>
+        <DefineConstants>__MonoCS__</DefineConstants>
+      </PropertyGroup>
+    </When>
   </Choose>
   <Choose>
     <When Condition="'$(OS)' == 'Windows_NT'">
-	  <PropertyGroup>
-		<Windows10SdkPath>$(Registry:HKEY_LOCAL_MACHINE\SOFTWARE\Wow6432Node\Microsoft\Microsoft SDKs\Windows\v10.0@InstallationFolder)</Windows10SdkPath>
-		<Windows10SdkLatestVersion>$(Registry:HKEY_LOCAL_MACHINE\SOFTWARE\Wow6432Node\Microsoft\Microsoft SDKs\Windows\v10.0@ProductVersion)</Windows10SdkLatestVersion>
-	  </PropertyGroup>
-	  <ItemGroup>
-		<Reference Include="Windows">
-		  <!-- This is a dummy entry to throw off the parser in DotNETUtilities's CsProjectInfo which doesn't understand MSBuild functions -->
-		  <HintPath/>      
-		  <HintPath>$(Windows10SdkPath)\UnionMetadata\$(Windows10SdkLatestVersion)\Windows.winmd</HintPath>
-		  <!-- 16299 SDK (possibly others?) has ProductVersion == 10.0.16299 in the registry, but pathing needs the final .0 -->
-		  <HintPath Condition="'$([System.IO.File]::Exists(%(HintPath)))'=='false'">$(Windows10SdkPath)\UnionMetadata\$(Windows10SdkLatestVersion).0\Windows.winmd</HintPath>
-		  <!-- Fall back to old, unversioned location -->
-		  <HintPath Condition="'$([System.IO.File]::Exists(%(HintPath)))'=='false'">$(Windows10SdkPath)\UnionMetadata\Windows.winmd</HintPath>
-		  <Private>false</Private>
-		</Reference>
-	  </ItemGroup>
-	</When>
+      <PropertyGroup>
+        <Windows10SdkPath>$(Registry:HKEY_LOCAL_MACHINE\SOFTWARE\Wow6432Node\Microsoft\Microsoft SDKs\Windows\v10.0@InstallationFolder)</Windows10SdkPath>
+        <Windows10SdkLatestVersion>$(Registry:HKEY_LOCAL_MACHINE\SOFTWARE\Wow6432Node\Microsoft\Microsoft SDKs\Windows\v10.0@ProductVersion)</Windows10SdkLatestVersion>
+      </PropertyGroup>
+      <ItemGroup>
+        <Reference Include="Windows">
+          <!-- This is a dummy entry to throw off the parser in DotNETUtilities's CsProjectInfo which doesn't understand MSBuild functions -->
+          <HintPath />
+          <HintPath>$(Windows10SdkPath)\UnionMetadata\$(Windows10SdkLatestVersion)\Windows.winmd</HintPath>
+          <!-- 16299 SDK (possibly others?) has ProductVersion == 10.0.16299 in the registry, but pathing needs the final .0 -->
+          <HintPath Condition="'$([System.IO.File]::Exists(%(HintPath)))'=='false'">$(Windows10SdkPath)\UnionMetadata\$(Windows10SdkLatestVersion).0\Windows.winmd</HintPath>
+          <!-- Fall back to old, unversioned location -->
+          <HintPath Condition="'$([System.IO.File]::Exists(%(HintPath)))'=='false'">$(Windows10SdkPath)\UnionMetadata\Windows.winmd</HintPath>
+          <Private>false</Private>
+        </Reference>
+      </ItemGroup>
+    </When>
   </Choose>
   <Choose>
     <When Condition="Exists('Platform\IOS\UEBuildIOS.cs')">
