Add a FTargetReceipt class to the engine, mirroring the UBT TargetReceipt class.

#rnx

diff --git a/Engine/Source/Developer/DesktopPlatform/Private/DesktopPlatformBase.cpp b/Engine/Source/Developer/DesktopPlatform/Private/DesktopPlatformBase.cpp
--- a/Engine/Source/Developer/DesktopPlatform/Private/DesktopPlatformBase.cpp
+++ b/Engine/Source/Developer/DesktopPlatform/Private/DesktopPlatformBase.cpp
@@ -1426,10 +1426,7 @@ bool FDesktopPlatformBase::ReadTargetInfo(const FString& FileName, TArray<FTarge
 		{
 			return false;
 		}
-
-		LexFromString(Targets[Idx].Type, *Type);
-
-		if (Targets[Idx].Type == EBuildTargetType::Unknown)
+		if(!LexTryParseString(Targets[Idx].Type, *Type) || Targets[Idx].Type == EBuildTargetType::Unknown)
 		{
 			return false;
 		}
diff --git a/Engine/Source/Developer/DesktopPlatform/Private/PlatformInfo.cpp b/Engine/Source/Developer/DesktopPlatform/Private/PlatformInfo.cpp
--- a/Engine/Source/Developer/DesktopPlatform/Private/PlatformInfo.cpp
+++ b/Engine/Source/Developer/DesktopPlatform/Private/PlatformInfo.cpp
@@ -145,7 +145,10 @@ void ParseDataDrivenPlatformInfo(const TCHAR* Name, const FConfigSection& Sectio
 	bool bTargetPlatformCanUseCrashReporter = TargetPlatformCanUseCrashReporterString.IsEmpty() ? true : GetSectionBool(Section, TEXT("bTargetPlatformCanUseCrashReporter"));
 
 	EBuildTargetType TargetType;
-	LexFromString(TargetType, *PlatformType);
+	if (!LexTryParseString(TargetType, *PlatformType))
+	{
+		TargetType = EBuildTargetType::Unknown;
+	}
 
 	BuildPlatformInfo(Name, TargetPlatformName, FText::FromString(DisplayName), TargetType, ConvertPlatformFlags(PlatformFlags), FPlatformIconPaths(NormalIconPath, LargeIconPath, XLargeIconPath), UATCommandLine,
 		AutoSDKPath, PlatformInfo::EPlatformSDKStatus::Unknown, TutorialPath, bIsEnabled, BinariesDirectoryName, IniPlatformName, bUsesHostCompiler, bUATClosesAfterLaunch, bIsConfidential, UBTTargetID, PlatformGroupName, bTargetPlatformCanUseCrashReporter);
diff --git a/Engine/Source/Developer/DesktopPlatform/Private/TargetReceipt.cpp b/Engine/Source/Developer/DesktopPlatform/Private/TargetReceipt.cpp
new file mode 100644
--- /dev/null
+++ b/Engine/Source/Developer/DesktopPlatform/Private/TargetReceipt.cpp
@@ -0,0 +1,161 @@
+// Copyright 1998-2019 Epic Games, Inc. All Rights Reserved.
+
+#include "TargetReceipt.h"
+#include "Misc/Paths.h"
+#include "Misc/FileHelper.h"
+#include "Dom/JsonObject.h"
+#include "Serialization/JsonSerializer.h"
+
+bool FTargetReceipt::Read(const FString& FileName)
+{
+	// Read the file from disk
+	FString FileContents;
+	if (!FFileHelper::LoadFileToString(FileContents, *FileName))
+	{
+		return false;
+	}
+
+	TSharedPtr<FJsonObject> Object;
+	TSharedRef<TJsonReader<>> Reader = TJsonReaderFactory<>::Create(FileContents);
+	if (!FJsonSerializer::Deserialize(Reader, Object) || !Object.IsValid())
+	{
+		return false;
+	}
+
+	// Get the project file
+	FString RelativeProjectFile;
+	if(Object->TryGetStringField(TEXT("Project"), RelativeProjectFile))
+	{
+		ProjectFile = FPaths::Combine(FPaths::GetPath(FileName), RelativeProjectFile);
+		FPaths::MakeStandardFilename(ProjectFile);
+		ProjectDir = FPaths::GetPath(ProjectFile);
+	}
+
+	// Get the target name
+	if (!Object->TryGetStringField(TEXT("TargetName"), TargetName))
+	{
+		return false;
+	}
+	if (!Object->TryGetStringField(TEXT("Platform"), Platform))
+	{
+		return false;
+	}
+	if (!Object->TryGetStringField(TEXT("Architecture"), Architecture))
+	{
+		return false;
+	}
+
+	// Read the configuration
+	FString ConfigurationString;
+	if (!Object->TryGetStringField(TEXT("Configuration"), ConfigurationString) || !LexTryParseString(Configuration, *ConfigurationString))
+	{
+		return false;
+	}
+
+	// Read the target type
+	FString TargetTypeString;
+	if (!Object->TryGetStringField(TEXT("TargetType"), TargetTypeString) || !LexTryParseString(TargetType, *TargetTypeString))
+	{
+		return false;
+	}
+
+	// Get the launch path
+	if(!Object->TryGetStringField(TEXT("Launch"), Launch))
+	{
+		return false;
+	}
+	ExpandVariables(Launch);
+
+	// Read the list of build products
+	const TArray<TSharedPtr<FJsonValue>>* BuildProductsArray;
+	if (Object->TryGetArrayField(TEXT("BuildProducts"), BuildProductsArray))
+	{
+		for(const TSharedPtr<FJsonValue>& BuildProductValue : *BuildProductsArray)
+		{
+			const TSharedPtr<FJsonObject>* BuildProductObject;
+			if(!BuildProductValue->TryGetObject(BuildProductObject))
+			{
+				return false;
+			}
+
+			FBuildProduct BuildProduct;
+			if(!(*BuildProductObject)->TryGetStringField(TEXT("Type"), BuildProduct.Type) || !(*BuildProductObject)->TryGetStringField(TEXT("Path"), BuildProduct.Path))
+			{
+				return false;
+			}
+
+			ExpandVariables(BuildProduct.Path);
+
+			BuildProducts.Add(MoveTemp(BuildProduct));
+		}
+	}
+
+	// Read the list of runtime dependencies
+	const TArray<TSharedPtr<FJsonValue>>* RuntimeDependenciesArray;
+	if (Object->TryGetArrayField(TEXT("RuntimeDependencies"), RuntimeDependenciesArray))
+	{
+		for(const TSharedPtr<FJsonValue>& RuntimeDependencyValue : *RuntimeDependenciesArray)
+		{
+			const TSharedPtr<FJsonObject>* RuntimeDependencyObject;
+			if(!RuntimeDependencyValue->TryGetObject(RuntimeDependencyObject))
+			{
+				return false;
+			}
+
+			FRuntimeDependency RuntimeDependency;
+			if(!(*RuntimeDependencyObject)->TryGetStringField(TEXT("Path"), RuntimeDependency.Path) || !(*RuntimeDependencyObject)->TryGetStringField(TEXT("Type"), RuntimeDependency.Type))
+			{
+				return false;
+			}
+
+			ExpandVariables(RuntimeDependency.Path);
+
+			RuntimeDependencies.Add(MoveTemp(RuntimeDependency));
+		}
+	}
+
+	// Read the list of additional properties
+	const TArray<TSharedPtr<FJsonValue>>* AdditionalPropertiesArray;
+	if (Object->TryGetArrayField(TEXT("AdditionalProperties"), AdditionalPropertiesArray))
+	{
+		for(const TSharedPtr<FJsonValue>& AdditionalPropertyValue : *AdditionalPropertiesArray)
+		{
+			const TSharedPtr<FJsonObject>* AdditionalPropertyObject;
+			if(!AdditionalPropertyValue->TryGetObject(AdditionalPropertyObject))
+			{
+				return false;
+			}
+
+			FReceiptProperty Property;
+			if(!(*AdditionalPropertyObject)->TryGetStringField(TEXT("Name"), Property.Name) || !(*AdditionalPropertyObject)->TryGetStringField(TEXT("Value"), Property.Value))
+			{
+				return false;
+			}
+
+			ExpandVariables(Property.Value);
+
+			AdditionalProperties.Add(MoveTemp(Property));
+		}
+	}
+
+	return true;
+}
+
+void FTargetReceipt::ExpandVariables(FString& Path)
+{
+	static FString EngineDirPrefix = TEXT("$(EngineDir)");
+	static FString ProjectDirPrefix = TEXT("$(ProjectDir)");
+	if(Path.StartsWith(EngineDirPrefix))
+	{
+		FString EngineDir = FPaths::EngineDir();
+		if (EngineDir.Len() > 0 && EngineDir[EngineDir.Len() - 1] == '/')
+		{
+			EngineDir = EngineDir.Left(EngineDir.Len() - 1);
+		}
+		Path = EngineDir + Path.Mid(EngineDirPrefix.Len());
+	}
+	else if(Path.StartsWith(ProjectDirPrefix) && ProjectDir.Len() > 0)
+	{
+		Path = ProjectDir + Path.Mid(ProjectDirPrefix.Len());
+	}
+}
diff --git a/Engine/Source/Developer/DesktopPlatform/Public/TargetReceipt.h b/Engine/Source/Developer/DesktopPlatform/Public/TargetReceipt.h
new file mode 100644
--- /dev/null
+++ b/Engine/Source/Developer/DesktopPlatform/Public/TargetReceipt.h
@@ -0,0 +1,97 @@
+// Copyright 1998-2019 Epic Games, Inc. All Rights Reserved.
+
+#pragma once
+
+#include "CoreMinimal.h"
+#include "Modules/BuildVersion.h"
+
+/**
+ * Record of a file that was created as part of the build process
+ */
+struct FBuildProduct
+{
+	/** Path to the file. */
+	FString Path;
+
+	/** Type of the build product. */
+	FString Type;
+};
+
+/** 
+ * Information about a file which is required by the target at runtime, and must be moved around with it.
+ */
+struct FRuntimeDependency
+{
+	/** The file that should be staged. Should use $(EngineDir) and $(ProjectDir) variables as a root, so that the target can be relocated to different machines. */
+	FString Path;
+
+	/** How to stage this file. */
+	FString Type;
+};
+
+/**
+ * Arbitrary property name/value which metadata from the build scripts can be passed on to downstream tasks
+ */
+struct FReceiptProperty
+{	
+	/** Property name */
+	FString Name;
+
+	/** Value of the property */
+	FString Value;
+};
+
+/**
+ * Stores information about a compiled target. Mirror of the TargetReceipt class generated by UBT.
+ */
+struct DESKTOPPLATFORM_API FTargetReceipt
+{
+	/** Path to the project file for this target */
+	FString ProjectFile;
+
+	/** The project directory */
+	FString ProjectDir;
+
+	/** The name of this target */
+	FString TargetName;
+
+	/** Which platform the target is compiled for */
+	FString Platform;
+
+	/** Which platform the target is compiled for */
+	FString Architecture;
+
+	/** Which configuration this target is compiled in */
+	EBuildConfiguration Configuration;
+
+	/** The type of the target */
+	EBuildTargetType TargetType;
+
+	/** The exectuable to launch for this target */
+	FString Launch;
+
+	/** The build products which are part of this target */
+	TArray<FBuildProduct> BuildProducts;
+
+	/** All the runtime dependencies that this target relies on */
+	TArray<FRuntimeDependency> RuntimeDependencies;
+
+	/** Additional build properties passed through from the module rules */
+	TArray<FReceiptProperty> AdditionalProperties;
+
+	/**
+	 * Read a target receipt from disk
+	 *
+	 * @param FileName The file to read from
+	 * @return True if the file was read successfully
+	 */
+	bool Read(const FString& FileName);
+
+private:
+	/**
+	 * Expands the $(EngineDir) and $(ProjectDir) variables within a string
+	 *
+	 * @param Path Path to expand variables within
+	 */
+	void ExpandVariables(FString& Path);
+};
diff --git a/Engine/Source/Developer/TargetPlatform/Private/InstalledPlatformInfo.cpp b/Engine/Source/Developer/TargetPlatform/Private/InstalledPlatformInfo.cpp
--- a/Engine/Source/Developer/TargetPlatform/Private/InstalledPlatformInfo.cpp
+++ b/Engine/Source/Developer/TargetPlatform/Private/InstalledPlatformInfo.cpp
@@ -58,11 +58,7 @@ void FInstalledPlatformInfo::ParsePlatformConfiguration(FString PlatformConfigur
 
 	FString ConfigurationName;
 	EBuildConfiguration Configuration = EBuildConfiguration::Unknown;
-	if (FParse::Value(*PlatformConfiguration, TEXT("Configuration="), ConfigurationName))
-	{
-		LexFromString(Configuration, *ConfigurationName);
-	}
-	if (Configuration == EBuildConfiguration::Unknown)
+	if (!FParse::Value(*PlatformConfiguration, TEXT("Configuration="), ConfigurationName) || !LexTryParseString(Configuration, *ConfigurationName) || Configuration == EBuildConfiguration::Unknown)
 	{
 		UE_LOG(LogInstalledPlatforms, Warning, TEXT("Unable to read configuration from %s"), *PlatformConfiguration);
 		bCanCreateEntry = false;
@@ -76,10 +72,10 @@ void FInstalledPlatformInfo::ParsePlatformConfiguration(FString PlatformConfigur
 	}
 
 	FString PlatformTypeName;
-	EBuildTargetType PlatformType = EBuildTargetType::Game;
-	if (FParse::Value(*PlatformConfiguration, TEXT("PlatformType="), PlatformTypeName))
+	EBuildTargetType PlatformType;
+	if (!FParse::Value(*PlatformConfiguration, TEXT("PlatformType="), PlatformTypeName) || !LexTryParseString(PlatformType, *PlatformTypeName))
 	{
-		LexFromString(PlatformType, *PlatformTypeName);
+		PlatformType = EBuildTargetType::Game;
 	}
 
 	FString Architecture;
diff --git a/Engine/Source/Runtime/Core/Private/GenericPlatform/GenericPlatformMisc.cpp b/Engine/Source/Runtime/Core/Private/GenericPlatform/GenericPlatformMisc.cpp
--- a/Engine/Source/Runtime/Core/Private/GenericPlatform/GenericPlatformMisc.cpp
+++ b/Engine/Source/Runtime/Core/Private/GenericPlatform/GenericPlatformMisc.cpp
@@ -52,31 +52,42 @@ CORE_API void (*ClipboardPasteShim)(FString& Dest) = nullptr;
 /* EBuildConfigurations interface
  *****************************************************************************/
 
-void LexFromString( EBuildConfiguration& OutConfiguration, const TCHAR* Configuration )
+bool LexTryParseString(EBuildConfiguration& OutConfiguration, const TCHAR* Configuration)
 {
 	if (FCString::Stricmp(Configuration, TEXT("Debug")) == 0)
 	{
 		OutConfiguration = EBuildConfiguration::Debug;
+		return true;
 	}
 	else if (FCString::Stricmp(Configuration, TEXT("DebugGame")) == 0)
 	{
 		OutConfiguration = EBuildConfiguration::DebugGame;
+		return true;
 	}
 	else if (FCString::Stricmp(Configuration, TEXT("Development")) == 0)
 	{
 		OutConfiguration = EBuildConfiguration::Development;
+		return true;
 	}
 	else if (FCString::Stricmp(Configuration, TEXT("Shipping")) == 0)
 	{
 		OutConfiguration = EBuildConfiguration::Shipping;
+		return true;
 	}
 	else if(FCString::Stricmp(Configuration, TEXT("Test")) == 0)
 	{
 		OutConfiguration = EBuildConfiguration::Test;
+		return true;
+	}
+	else if(FCString::Stricmp(Configuration, TEXT("Unknown")) == 0)
+	{
+		OutConfiguration = EBuildConfiguration::Unknown;
+		return true;
 	}
 	else
 	{
 		OutConfiguration = EBuildConfiguration::Unknown;
+		return false;
 	}
 }
 
@@ -104,7 +115,7 @@ namespace EBuildConfigurations
 	EBuildConfiguration FromString( const FString& Configuration )
 	{
 		EBuildConfiguration Result;
-		LexFromString(Result, *Configuration);
+		LexTryParseString(Result, *Configuration);
 		return Result;
 	}
 
@@ -142,31 +153,42 @@ namespace EBuildConfigurations
 /* EBuildTargetType functions
  *****************************************************************************/
 
-void LexFromString(EBuildTargetType& OutType, const TCHAR* Type)
+bool LexTryParseString(EBuildTargetType& OutType, const TCHAR* Type)
 {
 	if (FCString::Strcmp(Type, TEXT("Editor")) == 0)
 	{
 		OutType = EBuildTargetType::Editor;
+		return true;
 	}
 	else if (FCString::Strcmp(Type, TEXT("Game")) == 0)
 	{
 		OutType = EBuildTargetType::Game;
+		return true;
 	}
 	else if (FCString::Strcmp(Type, TEXT("Server")) == 0)
 	{
 		OutType = EBuildTargetType::Server;
+		return true;
 	}
 	else if (FCString::Strcmp(Type, TEXT("Client")) == 0)
 	{
 		OutType = EBuildTargetType::Client;
+		return true;
 	}
 	else if (FCString::Strcmp(Type, TEXT("Program")) == 0)
 	{
 		OutType = EBuildTargetType::Program;
+		return true;
+	}
+	else if (FCString::Strcmp(Type, TEXT("Unknown")) == 0)
+	{
+		OutType = EBuildTargetType::Unknown;
+		return true;
 	}
 	else
 	{
 		OutType = EBuildTargetType::Unknown;
+		return false;
 	}
 }
 
@@ -192,7 +214,7 @@ const TCHAR* LexToString(EBuildTargetType Type)
 EBuildTargetType EBuildTargets::FromString(const FString& Target)
 {
 	EBuildTargetType Type;
-	LexFromString(Type, *Target);
+	LexTryParseString(Type, *Target);
 	return Type;
 }
 
diff --git a/Engine/Source/Runtime/Core/Public/GenericPlatform/GenericPlatformMisc.h b/Engine/Source/Runtime/Core/Public/GenericPlatform/GenericPlatformMisc.h
--- a/Engine/Source/Runtime/Core/Public/GenericPlatform/GenericPlatformMisc.h
+++ b/Engine/Source/Runtime/Core/Public/GenericPlatform/GenericPlatformMisc.h
@@ -56,7 +56,7 @@ enum class EBuildConfiguration : uint8
  * @param Configuration The string to get the EBuildConfiguration for.
  * @return An EBuildConfiguration value.
  */
-CORE_API void LexFromString( EBuildConfiguration& OutConfiguration, const TCHAR* Configuration );
+CORE_API bool LexTryParseString(EBuildConfiguration& OutConfiguration, const TCHAR* Configuration);
 
 /**
  * Returns the string representation of the specified EBuildConfiguration value.
@@ -64,7 +64,7 @@ CORE_API void LexFromString( EBuildConfiguration& OutConfiguration, const TCHAR*
  * @param Configuration The value to get the string for.
  * @return The string representation.
  */
-CORE_API const TCHAR* LexToString( EBuildConfiguration Configuration );
+CORE_API const TCHAR* LexToString(EBuildConfiguration Configuration);
 
 namespace EBuildConfigurations
 {
@@ -147,7 +147,7 @@ enum class EBuildTargetType : uint8
  * @param Text The text to parse.
  * @return The string representation.
  */
-CORE_API void LexFromString(EBuildTargetType& OutType, const TCHAR* Text);
+CORE_API bool LexTryParseString(EBuildTargetType& OutType, const TCHAR* Text);
 
 /**
  * Returns the string representation of the specified EBuildTargetType value.
