Fix crash trying to format target receipt name in non-development configurations.

#rnx

diff --git a/Engine/Source/Developer/DesktopPlatform/Private/TargetReceipt.cpp b/Engine/Source/Developer/DesktopPlatform/Private/TargetReceipt.cpp
--- a/Engine/Source/Developer/DesktopPlatform/Private/TargetReceipt.cpp
+++ b/Engine/Source/Developer/DesktopPlatform/Private/TargetReceipt.cpp
@@ -155,7 +155,7 @@ FString FTargetReceipt::GetDefaultPath(const TCHAR* BaseDir, const TCHAR* Target
 	}
 	else
 	{
-		return FPaths::Combine(BaseDir, FString::Printf(TEXT("Binaries/%s/%s-%s-%s%s.target"), Platform, TargetName, Platform, Configuration, ArchitectureSuffix));
+		return FPaths::Combine(BaseDir, FString::Printf(TEXT("Binaries/%s/%s-%s-%s%s.target"), Platform, TargetName, Platform, LexToString(Configuration), ArchitectureSuffix));
 	}
 }
 
