Additional checks to prevent UBT being recompiled on build machines without being explicitly instructed to do so.

diff --git a/Engine/Source/Developer/DesktopPlatform/Private/DesktopPlatformBase.cpp b/Engine/Source/Developer/DesktopPlatform/Private/DesktopPlatformBase.cpp
--- a/Engine/Source/Developer/DesktopPlatform/Private/DesktopPlatformBase.cpp
+++ b/Engine/Source/Developer/DesktopPlatform/Private/DesktopPlatformBase.cpp
@@ -875,10 +875,17 @@ const TArray<FTargetInfo>& FDesktopPlatformBase::GetTargetsForProject(const FStr
 	}
 	Arguments += FString::Printf(TEXT(" -Output=\"%s\""), *IFileManager::Get().ConvertToAbsolutePathForExternalAppForRead(*InfoFileName));
 
-	// Run UBT to update the list of targets
+	// Run UBT to update the list of targets. Try to run it without building first.
 	FString Output;
-	int32 ReturnCode;
-	const_cast<FDesktopPlatformBase*>(this)->InvokeUnrealBuildToolSync(Arguments, *GLog, false, ReturnCode, Output);
+	int32 ReturnCode = 0;
+	if (!FPaths::FileExists(*GetUnrealBuildToolExecutableFilename(FPaths::RootDir())) || !const_cast<FDesktopPlatformBase*>(this)->InvokeUnrealBuildToolSync(Arguments, *GLog, true, ReturnCode, Output) || ReturnCode != 0)
+	{
+		// If that failed, try to build and run again. Build machines should always have an up-to-date copy, and shouldn't build anything without being told to do so.
+		if (!GIsBuildMachine)
+		{
+			const_cast<FDesktopPlatformBase*>(this)->InvokeUnrealBuildToolSync(Arguments, *GLog, false, ReturnCode, Output);
+		}
+	}
 
 	// Try to read the new targets
 	TArray<FTargetInfo> NewTargets;
