Remove few more ATG_CHANGE comments and questionable changes.

3/3(?)

diff --git a/Engine/Config/HoloLens/DataDrivenPlatformInfo.ini b/Engine/Config/HoloLens/DataDrivenPlatformInfo.ini
--- a/Engine/Config/HoloLens/DataDrivenPlatformInfo.ini
+++ b/Engine/Config/HoloLens/DataDrivenPlatformInfo.ini
@@ -18,6 +18,7 @@ Linux:bIsEnabled=false
 BinariesDirectoryName=HoloLens
 IniPlatformName=HoloLens
 bUsesHostCompiler=true
+bTargetPlatformCanUseCrashReporter=false;
 bUATClosesAfterLaunch=false
 UBTTargetID=HoloLens
 PlatformGroupName=Mobile
\ No newline at end of file
diff --git a/Engine/Source/Developer/DesktopPlatform/Private/PlatformInfo.cpp b/Engine/Source/Developer/DesktopPlatform/Private/PlatformInfo.cpp
--- a/Engine/Source/Developer/DesktopPlatform/Private/PlatformInfo.cpp
+++ b/Engine/Source/Developer/DesktopPlatform/Private/PlatformInfo.cpp
@@ -22,7 +22,7 @@ TArray<FPlatformInfo> AllPlatformInfoArray;
 // @todo platplug: Figure out why this is compiled on target devices
 #if WITH_EDITOR
 
-void BuildPlatformInfo(const FName& InPlatformInfoName, const FName& InTargetPlatformName, const FText& InDisplayName, const EPlatformType InPlatformType, const EPlatformFlags::Flags InPlatformFlags, const FPlatformIconPaths& InIconPaths, const FString& InUATCommandLine, const FString& InAutoSDKPath, EPlatformSDKStatus InStatus, const FString& InTutorial, bool InEnabled, FString InBinaryFolderName, FString InIniPlatformName, bool InUsesHostCompiler, bool InUATClosesAfterLaunch, bool InIsConfidential, const FName& InUBTTargetId, const FName& InPlatformGroupName)
+void BuildPlatformInfo(const FName& InPlatformInfoName, const FName& InTargetPlatformName, const FText& InDisplayName, const EPlatformType InPlatformType, const EPlatformFlags::Flags InPlatformFlags, const FPlatformIconPaths& InIconPaths, const FString& InUATCommandLine, const FString& InAutoSDKPath, EPlatformSDKStatus InStatus, const FString& InTutorial, bool InEnabled, FString InBinaryFolderName, FString InIniPlatformName, bool InUsesHostCompiler, bool InUATClosesAfterLaunch, bool InIsConfidential, const FName& InUBTTargetId, const FName& InPlatformGroupName, bool InTargetPlatformCanUseCrashReporter)
 {
 	FPlatformInfo& PlatformInfo = AllPlatformInfoArray.Emplace_GetRef();
 
@@ -79,13 +79,14 @@ void BuildPlatformInfo(const FName& InPlatformInfoName, const FName& InTargetPla
 	PlatformInfo.bUsesHostCompiler = InUsesHostCompiler;
 	PlatformInfo.bUATClosesAfterLaunch = InUATClosesAfterLaunch;
 	PlatformInfo.bIsConfidential = InIsConfidential;
+	PlatformInfo.bTargetPlatformCanUseCrashReporter = InTargetPlatformCanUseCrashReporter;
 }
 
 
 void BuildHardcodedPlatforms()
 {
-	// PlatformInfoName									TargetPlatformName			DisplayName														PlatformType			PlatformFlags					IconPaths																																		UATCommandLine										AutoSDKPath			SDKStatus						SDKTutorial																								bEnabledForUse										BinaryFolderName	IniPlatformName		FbUsesHostCompiler		bUATClosesAfterLaunch	bIsConfidential	UBTTargetId (match UBT's UnrealTargetPlatform enum)
-	BuildPlatformInfo(TEXT("AllDesktop"),				TEXT("AllDesktop"),			LOCTEXT("DesktopTargetPlatDisplay", "Desktop (Win+Mac+Linux)"),	EPlatformType::Game,	EPlatformFlags::None,			FPlatformIconPaths(TEXT("Launcher/Desktop/Platform_Desktop_24x"), TEXT("Launcher/Desktop/Platform_Desktop_128x")),					TEXT(""),											TEXT(""),			EPlatformSDKStatus::Unknown,	TEXT(""),																								PLATFORM_WINDOWS /* see note below */,						TEXT(""),			TEXT(""),			false,					true,					false,			TEXT("AllDesktop"),	TEXT("Desktop"));
+	// PlatformInfoName									TargetPlatformName			DisplayName														PlatformType			PlatformFlags					IconPaths																																		UATCommandLine										AutoSDKPath			SDKStatus						SDKTutorial																								bEnabledForUse										BinaryFolderName	IniPlatformName		FbUsesHostCompiler		bUATClosesAfterLaunch	bIsConfidential	UBTTargetId (match UBT's UnrealTargetPlatform enum)		bTargetPlatformCanUseCrashReporter
+	BuildPlatformInfo(TEXT("AllDesktop"),				TEXT("AllDesktop"),			LOCTEXT("DesktopTargetPlatDisplay", "Desktop (Win+Mac+Linux)"),	EPlatformType::Game,	EPlatformFlags::None,			FPlatformIconPaths(TEXT("Launcher/Desktop/Platform_Desktop_24x"), TEXT("Launcher/Desktop/Platform_Desktop_128x")),					TEXT(""),											TEXT(""),			EPlatformSDKStatus::Unknown,	TEXT(""),																								PLATFORM_WINDOWS /* see note below */,						TEXT(""),			TEXT(""),			false,					true,					false,			TEXT("AllDesktop"),	TEXT("Desktop"),	true);
 	// Note: For "AllDesktop" bEnabledForUse value, see SProjectTargetPlatformSettings::Construct !!!! IsAvailableOnWindows || IsAvailableOnMac || IsAvailableOnLinux
 }
 
@@ -140,9 +141,11 @@ void ParseDataDrivenPlatformInfo(const TCHAR* Name, const FConfigSection& Sectio
 	bool bIsConfidential = GetSectionBool(Section, TEXT("bIsConfidential"));
 	FName UBTTargetID = *GetSectionString(Section, TEXT("UBTTargetID"));
 	FName PlatformGroupName = *GetSectionString(Section, TEXT("PlatformGroupName"));
+	FString TargetPlatformCanUseCrashReporterString = GetSectionString(Section, TEXT("bTargetPlatformCanUseCrashReporter"));
+	bool bTargetPlatformCanUseCrashReporter = TargetPlatformCanUseCrashReporterString.IsEmpty() ? true : GetSectionBool(Section, TEXT("bTargetPlatformCanUseCrashReporter"));
 
 	BuildPlatformInfo(Name, TargetPlatformName, FText::FromString(DisplayName), EPlatformTypeFromString(*PlatformType), ConvertPlatformFlags(PlatformFlags), FPlatformIconPaths(NormalIconPath, LargeIconPath, XLargeIconPath), UATCommandLine,
-		AutoSDKPath, PlatformInfo::EPlatformSDKStatus::Unknown, TutorialPath, bIsEnabled, BinariesDirectoryName, IniPlatformName, bUsesHostCompiler, bUATClosesAfterLaunch, bIsConfidential, UBTTargetID, PlatformGroupName);
+		AutoSDKPath, PlatformInfo::EPlatformSDKStatus::Unknown, TutorialPath, bIsEnabled, BinariesDirectoryName, IniPlatformName, bUsesHostCompiler, bUATClosesAfterLaunch, bIsConfidential, UBTTargetID, PlatformGroupName, bTargetPlatformCanUseCrashReporter);
 }
 
 void LoadDataDrivenPlatforms()
diff --git a/Engine/Source/Developer/DesktopPlatform/Public/PlatformInfo.h b/Engine/Source/Developer/DesktopPlatform/Public/PlatformInfo.h
--- a/Engine/Source/Developer/DesktopPlatform/Public/PlatformInfo.h
+++ b/Engine/Source/Developer/DesktopPlatform/Public/PlatformInfo.h
@@ -172,6 +172,9 @@ namespace PlatformInfo
 		/** Whether or not the platform is confidential in nature */
 		bool bIsConfidential;
 
+		/** Whether or not the platform can use Crash Reporter */
+		bool bTargetPlatformCanUseCrashReporter;
+
 		/** An identifier that corresponds to UBT's UnrealTargetPlatform enum (and by proxy, FGenericPlatformMisc::GetUBTPlatform()) */
 		FName UBTTargetId;
 
diff --git a/Engine/Source/Editor/MainFrame/Private/Frame/MainFrameActions.cpp b/Engine/Source/Editor/MainFrame/Private/Frame/MainFrameActions.cpp
--- a/Engine/Source/Editor/MainFrame/Private/Frame/MainFrameActions.cpp
+++ b/Engine/Source/Editor/MainFrame/Private/Frame/MainFrameActions.cpp
@@ -669,8 +669,8 @@ void FMainFrameActionCallbacks::PackageProject( const FName InPlatformInfoName )
 		OptionalParams += TEXT(" -manifests");
 	}
 
-	bool bTargetPlatformCanUseCrashReporter = true;
-	if (PlatformInfo->TargetPlatformName == FName("WindowsNoEditor") && PlatformInfo->PlatformFlavor == TEXT("Win32"))
+	bool bTargetPlatformCanUseCrashReporter = PlatformInfo->bTargetPlatformCanUseCrashReporter;
+	if (bTargetPlatformCanUseCrashReporter && PlatformInfo->TargetPlatformName == FName("WindowsNoEditor") && PlatformInfo->PlatformFlavor == TEXT("Win32"))
 	{
 		FString MinumumSupportedWindowsOS;
 		GConfig->GetString(TEXT("/Script/WindowsTargetPlatform.WindowsTargetSettings"), TEXT("MinimumOSVersion"), MinumumSupportedWindowsOS, GEngineIni);
@@ -680,12 +680,7 @@ void FMainFrameActionCallbacks::PackageProject( const FName InPlatformInfoName )
 			bTargetPlatformCanUseCrashReporter = false;
 		}
 	}
-// @ATG_CHANGE : BEGIN HoloLens packaging & F5 support
-	if (PlatformInfo->TargetPlatformName == FName("HoloLens"))
-	{
-		bTargetPlatformCanUseCrashReporter = false;
-	}
-// @ATG_CHANGE : END
+
 	// Append any extra UAT flags specified for this platform flavor
 	if (!PlatformInfo->UATCommandLine.IsEmpty())
 	{
diff --git a/Engine/Source/Programs/DotNETCommon/DotNETUtilities/ResXResource/UEResXWriter.cs b/Engine/Source/Programs/DotNETCommon/DotNETUtilities/ResXResource/UEResXWriter.cs
--- a/Engine/Source/Programs/DotNETCommon/DotNETUtilities/ResXResource/UEResXWriter.cs
+++ b/Engine/Source/Programs/DotNETCommon/DotNETUtilities/ResXResource/UEResXWriter.cs
@@ -41,10 +41,7 @@ namespace Tools.DotNETCommon
 			XmlElement Value = Document.CreateElement("value");
 			Value.InnerText = InValue;
 
-			// @ATG_CHANGE : BEGIN
-			// Typo fix: Use parameter so that headers have resheader element name 
-			XmlElement Data = Document.CreateElement(InRootName);
-			// @ATG_CHANGE : END
+			XmlElement Data = Document.CreateElement("data");
 
 			{
 				XmlAttribute Attr = Document.CreateAttribute("name");
diff --git a/Engine/Source/Runtime/Core/Public/Misc/Optional.h b/Engine/Source/Runtime/Core/Public/Misc/Optional.h
--- a/Engine/Source/Runtime/Core/Public/Misc/Optional.h
+++ b/Engine/Source/Runtime/Core/Public/Misc/Optional.h
@@ -190,21 +190,6 @@ public:
 	const OptionalType& Get(const OptionalType& DefaultValue) const { return IsSet() ? *(OptionalType*)&Value : DefaultValue; }
 
 private:
-
-	// @ATG_CHANGE : BEGIN HoloLens support
-	// The approach of adjusting the packing mode around enum class forward declarations (because
-	// the compiler interprets these as WinRT types) is getting unwieldly, especially with IWYU 
-	// driving up the number of such declarations.  So instead we've just upped the default packing
-	// to 8 bytes on HoloLens32.  As a result, we need a different definition of TTypeCompatibleBytes to avoid
-	// C2719: 'parameter': formal parameter with __declspec(align('#')) won't be aligned.
-#if PLATFORM_HOLOLENS && !PLATFORM_64BITS
-	struct
-	{
-		uint8 Pad[sizeof(OptionalType)];
-	} Value;
-	bool bIsSet;
-#else
 	TTypeCompatibleBytes<OptionalType> Value;
 	bool bIsSet;
-#endif
 };
