Move Linux architectures to Linux package project submenu

Add a PlatformSubMenu ini entry to DataDrivenPlatformInfo sections, and coalesce these in FPackageProjectMenu::MakeMenu().

Cherry pick 9539016

diff --git a/Engine/Config/Linux/DataDrivenPlatformInfo.ini b/Engine/Config/Linux/DataDrivenPlatformInfo.ini
--- a/Engine/Config/Linux/DataDrivenPlatformInfo.ini
+++ b/Engine/Config/Linux/DataDrivenPlatformInfo.ini
@@ -22,6 +22,7 @@ Linux:bUsesHostCompiler=true
 bUATClosesAfterLaunch=true
 UBTTargetID=Linux
 PlatformGroupName=Desktop
+PlatformSubMenu=Linux
 
 [PlatformInfo LinuxNoEditor]
 TargetPlatformName=LinuxNoEditor
@@ -44,6 +45,7 @@ Linux:bUsesHostCompiler=true
 bUATClosesAfterLaunch=true
 UBTTargetID=Linux
 PlatformGroupName=Desktop
+PlatformSubMenu=Linux
 
 [PlatformInfo LinuxClient]
 TargetPlatformName=LinuxClient
@@ -66,6 +68,7 @@ Linux:bUsesHostCompiler=true
 bUATClosesAfterLaunch=true
 UBTTargetID=Linux
 PlatformGroupName=Desktop
+PlatformSubMenu=Linux
 
 [PlatformInfo LinuxServer]
 TargetPlatformName=LinuxServer
@@ -88,4 +91,5 @@ Linux:bUsesHostCompiler=true
 bUATClosesAfterLaunch=true
 UBTTargetID=Linux
 PlatformGroupName=Desktop
+PlatformSubMenu=Linux
 
diff --git a/Engine/Config/LinuxAArch64/DataDrivenPlatformInfo.ini b/Engine/Config/LinuxAArch64/DataDrivenPlatformInfo.ini
--- a/Engine/Config/LinuxAArch64/DataDrivenPlatformInfo.ini
+++ b/Engine/Config/LinuxAArch64/DataDrivenPlatformInfo.ini
@@ -23,3 +23,51 @@ Linux:bUsesHostCompiler=true
 bUATClosesAfterLaunch=true
 UBTTargetID=LinuxAArch64
 PlatformGroupName=Mobile
+PlatformSubMenu=Linux
+
+[PlatformInfo LinuxAArch64Client]
+TargetPlatformName=LinuxAArch64Client
+DisplayName=LinuxAArch64 (Client-only)
+PlatformType=Client
+PlatformFlags=None
+NormalIconPath=Launcher/Linux/Platform_Linux_24x
+LargeIconPath=Launcher/Linux/Platform_Linux_128x
+XLargeIconPath=
+UATCommandLine=-client
+AutoSDKPath=Linux_x64
+TutorialPath=/Platforms/Linux/GettingStarted
+Windows:bIsEnabled=true
+Mac:bIsEnabled=false
+Linux:bIsEnabled=true
+BinariesDirectoryName=LinuxAArch64
+IniPlatformName=LinuxAArch64
+Windows:bUsesHostCompiler=false
+Linux:bUsesHostCompiler=true
+bUATClosesAfterLaunch=true
+UBTTargetID=LinuxAArch64
+PlatformGroupName=Mobile
+PlatformSubMenu=Linux
+
+[PlatformInfo LinuxAArch64Server]
+TargetPlatformName=LinuxAArch64Server
+DisplayName=LinuxAArch64 (Dedicated Server)
+PlatformType=Server
+PlatformFlags=None
+NormalIconPath=Launcher/Linux/Platform_Linux_24x
+LargeIconPath=Launcher/Linux/Platform_Linux_128x
+XLargeIconPath=
+UATCommandLine=
+AutoSDKPath=Linux_x64
+TutorialPath=/Platforms/Linux/GettingStarted
+Windows:bIsEnabled=true
+Mac:bIsEnabled=false
+Linux:bIsEnabled=true
+BinariesDirectoryName=LinuxAArch64
+IniPlatformName=LinuxAArch64
+Windows:bUsesHostCompiler=false
+Linux:bUsesHostCompiler=true
+bUATClosesAfterLaunch=true
+UBTTargetID=LinuxAArch64
+PlatformGroupName=Mobile
+PlatformSubMenu=Linux
+
diff --git a/Engine/Source/Developer/DesktopPlatform/Private/PlatformInfo.cpp b/Engine/Source/Developer/DesktopPlatform/Private/PlatformInfo.cpp
--- a/Engine/Source/Developer/DesktopPlatform/Private/PlatformInfo.cpp
+++ b/Engine/Source/Developer/DesktopPlatform/Private/PlatformInfo.cpp
@@ -22,7 +22,7 @@ TArray<FPlatformInfo> AllPlatformInfoArray;
 // @todo platplug: Figure out why this is compiled on target devices
 #if WITH_EDITOR || IS_PROGRAM
 
-void BuildPlatformInfo(const FName& InPlatformInfoName, const FName& InTargetPlatformName, const FText& InDisplayName, const EBuildTargetType InPlatformType, const EPlatformFlags::Flags InPlatformFlags, const FPlatformIconPaths& InIconPaths, const FString& InUATCommandLine, const FString& InAutoSDKPath, EPlatformSDKStatus InStatus, const FString& InTutorial, bool InEnabled, FString InBinaryFolderName, FString InIniPlatformName, bool InUsesHostCompiler, bool InUATClosesAfterLaunch, bool InIsConfidential, const FName& InUBTTargetId, const FName& InPlatformGroupName, bool InTargetPlatformCanUseCrashReporter)
+void BuildPlatformInfo(const FName& InPlatformInfoName, const FName& InTargetPlatformName, const FText& InDisplayName, const EBuildTargetType InPlatformType, const EPlatformFlags::Flags InPlatformFlags, const FPlatformIconPaths& InIconPaths, const FString& InUATCommandLine, const FString& InAutoSDKPath, EPlatformSDKStatus InStatus, const FString& InTutorial, bool InEnabled, FString InBinaryFolderName, FString InIniPlatformName, bool InUsesHostCompiler, bool InUATClosesAfterLaunch, bool InIsConfidential, const FName& InUBTTargetId, const FName& InPlatformGroupName, const FName& InPlatformSubMenu, bool InTargetPlatformCanUseCrashReporter)
 {
 	FPlatformInfo& PlatformInfo = AllPlatformInfoArray.Emplace_GetRef();
 
@@ -59,6 +59,7 @@ void BuildPlatformInfo(const FName& InPlatformInfoName, const FName& InTargetPla
 	PlatformInfo.IniPlatformName = InIniPlatformName;
 	PlatformInfo.UBTTargetId = InUBTTargetId;
 	PlatformInfo.PlatformGroupName = InPlatformGroupName;
+	PlatformInfo.PlatformSubMenu = InPlatformSubMenu;
 
 	if (InPlatformGroupName != NAME_None)
 	{
@@ -85,8 +86,8 @@ void BuildPlatformInfo(const FName& InPlatformInfoName, const FName& InTargetPla
 
 void BuildHardcodedPlatforms()
 {
-	// PlatformInfoName									TargetPlatformName			DisplayName														PlatformType			PlatformFlags					IconPaths																																		UATCommandLine										AutoSDKPath			SDKStatus						SDKTutorial																								bEnabledForUse										BinaryFolderName	IniPlatformName		FbUsesHostCompiler		bUATClosesAfterLaunch	bIsConfidential	UBTTargetId (match UBT's UnrealTargetPlatform enum)		bTargetPlatformCanUseCrashReporter
-	BuildPlatformInfo(TEXT("AllDesktop"),				TEXT("AllDesktop"),			LOCTEXT("DesktopTargetPlatDisplay", "Desktop (Win+Mac+Linux)"),	EBuildTargetType::Game,	EPlatformFlags::None,			FPlatformIconPaths(TEXT("Launcher/Desktop/Platform_Desktop_24x"), TEXT("Launcher/Desktop/Platform_Desktop_128x")),					TEXT(""),											TEXT(""),			EPlatformSDKStatus::Unknown,	TEXT(""),																								PLATFORM_WINDOWS /* see note below */,						TEXT(""),			TEXT(""),			false,					true,					false,			TEXT("AllDesktop"),	TEXT("Desktop"),	true);
+	// PlatformInfoName									TargetPlatformName			DisplayName														PlatformType			PlatformFlags					IconPaths																																		UATCommandLine										AutoSDKPath			SDKStatus						SDKTutorial																								bEnabledForUse										BinaryFolderName	IniPlatformName		FbUsesHostCompiler		bUATClosesAfterLaunch	bIsConfidential	UBTTargetId (match UBT's UnrealTargetPlatform enum)		PlatformSubMenu		bTargetPlatformCanUseCrashReporter
+	BuildPlatformInfo(TEXT("AllDesktop"),				TEXT("AllDesktop"),			LOCTEXT("DesktopTargetPlatDisplay", "Desktop (Win+Mac+Linux)"),	EBuildTargetType::Game,	EPlatformFlags::None,			FPlatformIconPaths(TEXT("Launcher/Desktop/Platform_Desktop_24x"), TEXT("Launcher/Desktop/Platform_Desktop_128x")),					TEXT(""),											TEXT(""),			EPlatformSDKStatus::Unknown,	TEXT(""),																								PLATFORM_WINDOWS /* see note below */,						TEXT(""),			TEXT(""),			false,					true,					false,			TEXT("AllDesktop"),	TEXT("Desktop"),	TEXT(""),	true);
 	// Note: For "AllDesktop" bEnabledForUse value, see SProjectTargetPlatformSettings::Construct !!!! IsAvailableOnWindows || IsAvailableOnMac || IsAvailableOnLinux
 }
 
@@ -135,6 +136,7 @@ namespace EPlatformSection
 	const FName bIsConfidential = FName(TEXT("bIsConfidential"));
 	const FName UBTTargetID = FName(TEXT("UBTTargetID"));
 	const FName PlatformGroupName = FName(TEXT("PlatformGroupName"));
+	const FName PlatformSubMenu = FName(TEXT("PlatformSubMenu"));
 }
 
 void ParseDataDrivenPlatformInfo(const TCHAR* Name, const FConfigSection& Section)
@@ -162,6 +164,7 @@ void ParseDataDrivenPlatformInfo(const TCHAR* Name, const FConfigSection& Sectio
 	bool bIsConfidential = GetSectionBool(Section, EPlatformSection::bIsConfidential);
 	FName UBTTargetID = *GetSectionString(Section, EPlatformSection::UBTTargetID);
 	FName PlatformGroupName = *GetSectionString(Section, EPlatformSection::PlatformGroupName);
+	FName PlatformSubMenu = *GetSectionString(Section, EPlatformSection::PlatformSubMenu);
 	FString TargetPlatformCanUseCrashReporterString = GetSectionString(Section, TEXT("bTargetPlatformCanUseCrashReporter"));
 	bool bTargetPlatformCanUseCrashReporter = TargetPlatformCanUseCrashReporterString.IsEmpty() ? true : GetSectionBool(Section, TEXT("bTargetPlatformCanUseCrashReporter"));
 
@@ -172,7 +175,7 @@ void ParseDataDrivenPlatformInfo(const TCHAR* Name, const FConfigSection& Sectio
 	}
 
 	BuildPlatformInfo(Name, TargetPlatformName, FText::FromString(DisplayName), TargetType, ConvertPlatformFlags(PlatformFlags), FPlatformIconPaths(NormalIconPath, LargeIconPath, XLargeIconPath), UATCommandLine,
-		AutoSDKPath, PlatformInfo::EPlatformSDKStatus::Unknown, TutorialPath, bIsEnabled, BinariesDirectoryName, IniPlatformName, bUsesHostCompiler, bUATClosesAfterLaunch, bIsConfidential, UBTTargetID, PlatformGroupName, bTargetPlatformCanUseCrashReporter);
+		AutoSDKPath, PlatformInfo::EPlatformSDKStatus::Unknown, TutorialPath, bIsEnabled, BinariesDirectoryName, IniPlatformName, bUsesHostCompiler, bUATClosesAfterLaunch, bIsConfidential, UBTTargetID, PlatformGroupName, PlatformSubMenu, bTargetPlatformCanUseCrashReporter);
 }
 
 void LoadDataDrivenPlatforms()
diff --git a/Engine/Source/Developer/DesktopPlatform/Public/PlatformInfo.h b/Engine/Source/Developer/DesktopPlatform/Public/PlatformInfo.h
--- a/Engine/Source/Developer/DesktopPlatform/Public/PlatformInfo.h
+++ b/Engine/Source/Developer/DesktopPlatform/Public/PlatformInfo.h
@@ -165,6 +165,9 @@ namespace PlatformInfo
 		/** An identifier to group similar platforms together, such as "Mobile" and "Console". Used for Per-Platform Override Properties. */
 		FName PlatformGroupName;
 
+		/** Submenu name to group similar platforms together in menus, such as "Linux" and "LinuxAArch64".  */
+		FName PlatformSubMenu;
+
 		/** Returns true if this platform is vanilla */
 		FORCEINLINE bool IsVanilla() const
 		{
diff --git a/Engine/Source/Editor/MainFrame/Private/Menus/PackageProjectMenu.h b/Engine/Source/Editor/MainFrame/Private/Menus/PackageProjectMenu.h
--- a/Engine/Source/Editor/MainFrame/Private/Menus/PackageProjectMenu.h
+++ b/Engine/Source/Editor/MainFrame/Private/Menus/PackageProjectMenu.h
@@ -34,6 +34,7 @@ public:
 	 */
 	static void MakeMenu( FMenuBuilder& MenuBuilder )
 	{
+		TArray<FName> AllPlatformSubMenus;
 		const TArray<FString>& ConfidentalPlatforms = FDataDrivenPlatformInfoRegistry::GetConfidentialPlatforms();
 
 		TArray<PlatformInfo::FVanillaPlatformEntry> VanillaPlatforms = PlatformInfo::BuildPlatformHierarchy(PlatformInfo::EPlatformFilter::All);
@@ -67,7 +68,40 @@ public:
 				continue;
 			}
 
-			if (VanillaPlatform.PlatformFlavors.Num())
+			// Check if this platform has a submenu entry
+			if (VanillaPlatform.PlatformInfo->PlatformSubMenu != NAME_None)
+			{
+				TArray<const PlatformInfo::FPlatformInfo*> SubMenuEntries;
+				const FName& PlatformSubMenu = VanillaPlatform.PlatformInfo->PlatformSubMenu;
+
+				// Check if we've already added this submenu
+				if (AllPlatformSubMenus.Find(PlatformSubMenu) != INDEX_NONE)
+					continue;
+				AllPlatformSubMenus.Add(PlatformSubMenu);
+
+				// Go through all vanilla platforms looking for matching submenus
+				for (const PlatformInfo::FVanillaPlatformEntry& SubMenuVanillaPlatform : VanillaPlatforms)
+				{
+					const PlatformInfo::FPlatformInfo* PlatformInfo = SubMenuVanillaPlatform.PlatformInfo;
+
+					if ((PlatformInfo->PlatformType == EBuildTargetType::Game) && (PlatformInfo->PlatformSubMenu == PlatformSubMenu))
+					{
+						SubMenuEntries.Add(PlatformInfo);
+					}
+				}
+
+				if (SubMenuEntries.Num())
+				{
+					const FText DisplayName = FText::FromName(PlatformSubMenu);
+
+					MenuBuilder.AddSubMenu(
+							ProjectTargetPlatformEditorModule.MakePlatformMenuItemWidget(*VanillaPlatform.PlatformInfo, false, DisplayName), 
+							FNewMenuDelegate::CreateStatic(&FPackageProjectMenu::AddPlatformSubPlatformsToMenu, SubMenuEntries),
+							false
+							);
+				}
+			}
+			else if (VanillaPlatform.PlatformFlavors.Num())
 			{
 				MenuBuilder.AddSubMenu(
 					ProjectTargetPlatformEditorModule.MakePlatformMenuItemWidget(*VanillaPlatform.PlatformInfo), 
