Fix deprecated ARRAY_COUNT -> UE_ARRAY_COUNT

diff --git a/Engine/Source/Developer/DesktopPlatform/Private/DesktopPlatformBase.cpp b/Engine/Source/Developer/DesktopPlatform/Private/DesktopPlatformBase.cpp
--- a/Engine/Source/Developer/DesktopPlatform/Private/DesktopPlatformBase.cpp
+++ b/Engine/Source/Developer/DesktopPlatform/Private/DesktopPlatformBase.cpp
@@ -774,7 +774,7 @@ struct FTargetFileVisitor : IPlatformFile::FDirectoryStatVisitor
 		int32 Length = FCString::Strlen(FileNameOrDirectory);
 
 		static const TCHAR TargetExt[] = TEXT(".Target.cs");
-		static const int32 TargetExtLen = ARRAY_COUNT(TargetExt) - 1;
+		static const int32 TargetExtLen = UE_ARRAY_COUNT(TargetExt) - 1;
 		if (Length > TargetExtLen && FCString::Stricmp(FileNameOrDirectory + Length - TargetExtLen, TargetExt) == 0)
 		{
 			FString TargetName = FPaths::GetCleanFilename(FString(Length - TargetExtLen, FileNameOrDirectory));
@@ -782,7 +782,7 @@ struct FTargetFileVisitor : IPlatformFile::FDirectoryStatVisitor
 		}
 		
 		static const TCHAR ModuleExt[] = TEXT(".Build.cs");
-		static const int32 ModuleExtLen = ARRAY_COUNT(ModuleExt) - 1;
+		static const int32 ModuleExtLen = UE_ARRAY_COUNT(ModuleExt) - 1;
 		if (Length > ModuleExtLen && FCString::Stricmp(FileNameOrDirectory + Length - ModuleExtLen, ModuleExt) == 0)
 		{
 			bSearchSubDirectories = false;
