Speculative CIS fix.

#rnx

diff --git a/Engine/Source/Developer/DesktopPlatform/Private/Linux/DesktopPlatformLinux.h b/Engine/Source/Developer/DesktopPlatform/Private/Linux/DesktopPlatformLinux.h
--- a/Engine/Source/Developer/DesktopPlatform/Private/Linux/DesktopPlatformLinux.h
+++ b/Engine/Source/Developer/DesktopPlatform/Private/Linux/DesktopPlatformLinux.h
@@ -25,6 +25,7 @@ public:
 
 	virtual bool OpenProject(const FString &ProjectFileName) override;
 
+	using FDesktopPlatformBase::RunUnrealBuildTool;
 	virtual bool RunUnrealBuildTool(const FText& Description, const FString& RootDir, const FString& Arguments, FFeedbackContext* Warn, int32& OutExitCode) override;
 	virtual bool IsUnrealBuildToolRunning() override;
 
diff --git a/Engine/Source/Developer/DesktopPlatform/Private/Mac/DesktopPlatformMac.h b/Engine/Source/Developer/DesktopPlatform/Private/Mac/DesktopPlatformMac.h
--- a/Engine/Source/Developer/DesktopPlatform/Private/Mac/DesktopPlatformMac.h
+++ b/Engine/Source/Developer/DesktopPlatform/Private/Mac/DesktopPlatformMac.h
@@ -20,6 +20,7 @@ public:
 	virtual bool VerifyFileAssociations() override;
 	virtual bool UpdateFileAssociations() override;
 
+	using FDesktopPlatformBase::RunUnrealBuildTool;
 	virtual bool RunUnrealBuildTool(const FText& Description, const FString& RootDir, const FString& Arguments, FFeedbackContext* Warn, int32& OutExitCode) override;
 	virtual bool IsUnrealBuildToolRunning() override;
 
