Fix compile error.

#rnx

diff --git a/Engine/Source/Developer/DesktopPlatform/Private/DesktopPlatformStub.h b/Engine/Source/Developer/DesktopPlatform/Private/DesktopPlatformStub.h
--- a/Engine/Source/Developer/DesktopPlatform/Private/DesktopPlatformStub.h
+++ b/Engine/Source/Developer/DesktopPlatform/Private/DesktopPlatformStub.h
@@ -22,7 +22,7 @@ public:
 	virtual bool VerifyFileAssociations() override { return false; }
 	virtual bool UpdateFileAssociations() override { return false; }
 
-	virtual bool RunUnrealBuildTool(const FText& Description, const FString& RootDir, const FString& Arguments, FFeedbackContext* Warn) override { return false; }
+	virtual bool RunUnrealBuildTool(const FText& Description, const FString& RootDir, const FString& Arguments, FFeedbackContext* Warn, int32& OutExitCode) override { OutExitCode = 1;  return false; }
 	virtual bool IsUnrealBuildToolRunning() override { return false; }
 
 	virtual FFeedbackContext* GetNativeFeedbackContext() override { return nullptr; }
