Fix Clang compile error on Windows.

#rnx

diff --git a/Engine/Source/Developer/DesktopPlatform/Private/Windows/DesktopPlatformWindows.h b/Engine/Source/Developer/DesktopPlatform/Private/Windows/DesktopPlatformWindows.h
--- a/Engine/Source/Developer/DesktopPlatform/Private/Windows/DesktopPlatformWindows.h
+++ b/Engine/Source/Developer/DesktopPlatform/Private/Windows/DesktopPlatformWindows.h
@@ -28,6 +28,7 @@ public:
 
 	virtual bool OpenProject(const FString &ProjectFileName) override;
 
+	using FDesktopPlatformBase::RunUnrealBuildTool;
 	virtual bool RunUnrealBuildTool(const FText& Description, const FString& RootDir, const FString& Arguments, FFeedbackContext* Warn, int32& OutExitCode) override;
 	virtual bool IsUnrealBuildToolRunning() override;
 
