Parent the DDC download window to the splash screen, if it's active.

diff --git a/Engine/Source/Developer/DesktopPlatform/Private/Windows/WindowsNativeFeedbackContext.cpp b/Engine/Source/Developer/DesktopPlatform/Private/Windows/WindowsNativeFeedbackContext.cpp
--- a/Engine/Source/Developer/DesktopPlatform/Private/Windows/WindowsNativeFeedbackContext.cpp
+++ b/Engine/Source/Developer/DesktopPlatform/Private/Windows/WindowsNativeFeedbackContext.cpp
@@ -173,6 +173,9 @@ DWORD FWindowsNativeFeedbackContext::SlowTaskThreadProc(void* ThreadParam)
 
 	HINSTANCE HInstance = (HINSTANCE)GetModuleHandle(NULL);
 
+	extern APPLICATIONCORE_API HWND GetSplashScreenWindowHandle();
+	HWND ParentWindowHandle = GetSplashScreenWindowHandle();
+
 	WNDCLASSEX WndClassEx;
 	ZeroMemory(&WndClassEx, sizeof(WndClassEx));
 	WndClassEx.cbSize = sizeof(WndClassEx);
@@ -219,7 +222,7 @@ DWORD FWindowsNativeFeedbackContext::SlowTaskThreadProc(void* ThreadParam)
 	AdjustWindowRectEx(&WindowRect, WindowStyle, 0, 0);
 
 	const TCHAR* WindowClassName = MAKEINTATOM( WndClassAtom );
-	HWND hWnd = CreateWindow(WindowClassName, TEXT("Unreal Engine"), WindowStyle, WindowRect.left, WindowRect.top, WindowRect.right - WindowRect.left, WindowRect.bottom - WindowRect.top, NULL, NULL, HInstance, NULL);
+	HWND hWnd = CreateWindow(WindowClassName, TEXT("Unreal Engine"), WindowStyle, WindowRect.left, WindowRect.top, WindowRect.right - WindowRect.left, WindowRect.bottom - WindowRect.top, ParentWindowHandle, NULL, HInstance, NULL);
 	SetWindowLongPtr(hWnd, GWLP_USERDATA, (LONG_PTR)&Params);
 	SendMessageW(hWnd, WM_SETFONT, (WPARAM)hFont, 0);
 
diff --git a/Engine/Source/Runtime/ApplicationCore/Private/Windows/WindowsPlatformSplash.cpp b/Engine/Source/Runtime/ApplicationCore/Private/Windows/WindowsPlatformSplash.cpp
--- a/Engine/Source/Runtime/ApplicationCore/Private/Windows/WindowsPlatformSplash.cpp
+++ b/Engine/Source/Runtime/ApplicationCore/Private/Windows/WindowsPlatformSplash.cpp
@@ -764,4 +764,9 @@ void FWindowsPlatformSplash::SetSplashText( const SplashTextType::Type InType, c
 	}
 }
 
+APPLICATIONCORE_API HWND GetSplashScreenWindowHandle()
+{
+	return GSplashScreenWnd;
+}
+
 #include "Windows/HideWindowsPlatformTypes.h"
