#UE4 Promoted builds may not have a source code folder. Look for TargetInfo.json file before assuming no targets when there is no Source folder. Also assume all targets in the file are valid since there is no way to verify them if the code is missing.

diff --git a/Engine/Source/Developer/DesktopPlatform/Private/DesktopPlatformBase.cpp b/Engine/Source/Developer/DesktopPlatform/Private/DesktopPlatformBase.cpp
--- a/Engine/Source/Developer/DesktopPlatform/Private/DesktopPlatformBase.cpp
+++ b/Engine/Source/Developer/DesktopPlatform/Private/DesktopPlatformBase.cpp
@@ -806,6 +806,12 @@ struct FTargetFileVisitor : IPlatformFile::FDirectoryStatVisitor
 
 bool IsTargetInfoValid(const TArray<FTargetInfo>& Targets, const FString& SourceDir, const FDateTime& LastModifiedTime)
 {
+	if (FApp::GetEngineIsPromotedBuild())
+	{
+		// Promoted builds may not have source code, so we will assume all supplied targets are valid since they will not appear on disk
+		return true;
+	}
+
 	// Create the state 
 	TSet<FString> RemainingTargetNames;
 	for (const FTargetInfo& Target : Targets)
@@ -856,12 +862,7 @@ const TArray<FTargetInfo>& FDesktopPlatformBase::GetTargetsForProject(const FStr
 		ProjectDir = FPaths::GetPath(ProjectFile);
 	}
 
-	// Get the project source directory. If it doesn't exist, there are no targets for this project.
 	FString ProjectSourceDir = ProjectDir / TEXT("Source");
-	if (!IFileManager::Get().DirectoryExists(*ProjectSourceDir))
-	{
-		return ProjectFileToTargets.Add(NormalizedProjectFile, TArray<FTargetInfo>());
-	}
 
 	// Get the path to the info filename
 	FString InfoFileName = ProjectDir / TEXT("Intermediate/TargetInfo.json");
@@ -878,6 +879,12 @@ const TArray<FTargetInfo>& FDesktopPlatformBase::GetTargetsForProject(const FStr
 		}
 	}
 
+	// Get the project source directory. If it doesn't exist, there are no targets for this project.
+	if (!IFileManager::Get().DirectoryExists(*ProjectSourceDir))
+	{
+		return ProjectFileToTargets.Add(NormalizedProjectFile, TArray<FTargetInfo>());
+	}
+
 	// Otherwise, we'll have to run UBT to update it
 	FString Arguments = TEXT("-Mode=QueryTargets");
 	if(ProjectFile.Len() > 0)
