Added ini option to control path to UBT as spawned by editor.
Defaults to old UBT path. Will be used to override path for netcore built version of UBT.

diff --git a/Engine/Source/Developer/DesktopPlatform/Private/DesktopPlatformBase.cpp b/Engine/Source/Developer/DesktopPlatform/Private/DesktopPlatformBase.cpp
--- a/Engine/Source/Developer/DesktopPlatform/Private/DesktopPlatformBase.cpp
+++ b/Engine/Source/Developer/DesktopPlatform/Private/DesktopPlatformBase.cpp
@@ -1469,6 +1469,13 @@ FString FDesktopPlatformBase::GetUnrealBuildToolProjectFileName(const FString& R
 
 FString FDesktopPlatformBase::GetUnrealBuildToolExecutableFilename(const FString& RootDir) const
 {
+	FString Entry;
+	if( GConfig->GetString( TEXT("PlatformPaths"), TEXT("UnrealBuildTool"), Entry, GEngineIni ))
+	{
+		FString NewPath = FPaths::ConvertRelativePathToFull(RootDir / Entry);
+		return NewPath;
+	}
+
 	return FPaths::ConvertRelativePathToFull(RootDir / TEXT("Engine/Binaries/DotNET/UnrealBuildTool.exe"));
 }
 
