Removed redundant paths to UBT controlling it in a single place that is aware of our netcore locations. Fixes issues creating new foreign projects.

diff --git a/Engine/Source/Developer/DesktopPlatform/Private/Linux/DesktopPlatformLinux.cpp b/Engine/Source/Developer/DesktopPlatform/Private/Linux/DesktopPlatformLinux.cpp
--- a/Engine/Source/Developer/DesktopPlatform/Private/Linux/DesktopPlatformLinux.cpp
+++ b/Engine/Source/Developer/DesktopPlatform/Private/Linux/DesktopPlatformLinux.cpp
@@ -419,7 +419,7 @@ bool FDesktopPlatformLinux::RunUnrealBuildTool(const FText& Description, const F
 	OutExitCode = 1;
 
 	// Get the path to UBT
-	FString UnrealBuildToolPath = RootDir / TEXT("Engine/Binaries/DotNET/UnrealBuildTool.exe");
+	FString UnrealBuildToolPath = GetUnrealBuildToolExecutableFilename(RootDir);
 	if(IFileManager::Get().FileSize(*UnrealBuildToolPath) < 0)
 	{
 		Warn->Logf(ELogVerbosity::Error, TEXT("Couldn't find UnrealBuildTool at '%s'"), *UnrealBuildToolPath);
diff --git a/Engine/Source/Developer/DesktopPlatform/Private/Mac/DesktopPlatformMac.cpp b/Engine/Source/Developer/DesktopPlatform/Private/Mac/DesktopPlatformMac.cpp
--- a/Engine/Source/Developer/DesktopPlatform/Private/Mac/DesktopPlatformMac.cpp
+++ b/Engine/Source/Developer/DesktopPlatform/Private/Mac/DesktopPlatformMac.cpp
@@ -602,7 +602,7 @@ bool FDesktopPlatformMac::RunUnrealBuildTool(const FText& Description, const FSt
 	OutExitCode = 1;
 
 	// Get the path to UBT
-	FString UnrealBuildToolPath = RootDir / TEXT("Engine/Binaries/DotNET/UnrealBuildTool.exe");
+	FString UnrealBuildToolPath = GetUnrealBuildToolExecutableFilename(RootDir);
 	if(IFileManager::Get().FileSize(*UnrealBuildToolPath) < 0)
 	{
 		Warn->Logf(ELogVerbosity::Error, TEXT("Couldn't find UnrealBuildTool at '%s'"), *UnrealBuildToolPath);
diff --git a/Engine/Source/Developer/DesktopPlatform/Private/Windows/DesktopPlatformWindows.cpp b/Engine/Source/Developer/DesktopPlatform/Private/Windows/DesktopPlatformWindows.cpp
--- a/Engine/Source/Developer/DesktopPlatform/Private/Windows/DesktopPlatformWindows.cpp
+++ b/Engine/Source/Developer/DesktopPlatform/Private/Windows/DesktopPlatformWindows.cpp
@@ -463,7 +463,7 @@ bool FDesktopPlatformWindows::RunUnrealBuildTool(const FText& Description, const
 	OutExitCode = 1;
 
 	// Get the path to UBT
-	FString UnrealBuildToolPath = RootDir / TEXT("Engine/Binaries/DotNET/UnrealBuildTool.exe");
+	FString UnrealBuildToolPath = GetUnrealBuildToolExecutableFilename(RootDir);
 	if(IFileManager::Get().FileSize(*UnrealBuildToolPath) < 0)
 	{
 		Warn->Logf(ELogVerbosity::Error, TEXT("Couldn't find UnrealBuildTool at '%s'"), *UnrealBuildToolPath);
