Fixes for Mac/Linux editors spawning UBT

diff --git a/Engine/Config/Linux/BaseLinuxEngine.ini b/Engine/Config/Linux/BaseLinuxEngine.ini
new file mode 100644
--- /dev/null
+++ b/Engine/Config/Linux/BaseLinuxEngine.ini
@@ -0,0 +1,2 @@
+[PlatformPaths]
+UnrealBuildTool=Engine/Binaries/DotNET/UnrealBuildTool/UnrealBuildTool
\ No newline at end of file
diff --git a/Engine/Config/Mac/BaseMacEngine.ini b/Engine/Config/Mac/BaseMacEngine.ini
--- a/Engine/Config/Mac/BaseMacEngine.ini
+++ b/Engine/Config/Mac/BaseMacEngine.ini
@@ -3,3 +3,6 @@ AudioCaptureModuleName=AudioCaptureRtAudio
 
 [OnlineSubsystem]
 LocalPlatformName=MAC
+
+[PlatformPaths]
+UnrealBuildTool=Engine/Binaries/DotNET/UnrealBuildTool/UnrealBuildTool
\ No newline at end of file
diff --git a/Engine/Source/Developer/DesktopPlatform/Private/DesktopPlatformBase.cpp b/Engine/Source/Developer/DesktopPlatform/Private/DesktopPlatformBase.cpp
--- a/Engine/Source/Developer/DesktopPlatform/Private/DesktopPlatformBase.cpp
+++ b/Engine/Source/Developer/DesktopPlatform/Private/DesktopPlatformBase.cpp
@@ -729,18 +729,6 @@ FProcHandle FDesktopPlatformBase::InvokeUnrealBuildToolAsync(const FString& InCm
 
 	Ar.Logf(TEXT("Launching UnrealBuildTool... [%s %s]"), *ExecutableFileName, *CmdLineParams);
 
-#if PLATFORM_MAC
-	// On Mac we launch UBT with Mono
-	FString ScriptPath = FPaths::ConvertRelativePathToFull(FPaths::EngineDir() / TEXT("Build/BatchFiles/Mac/RunMono.sh"));
-	CmdLineParams = FString::Printf(TEXT("\"%s\" \"%s\" %s"), *ScriptPath, *ExecutableFileName, *CmdLineParams);
-	ExecutableFileName = TEXT("/bin/sh");
-#elif PLATFORM_LINUX
-	// Real men run Linux (with Mono??)
-	FString ScriptPath = FPaths::ConvertRelativePathToFull(FPaths::EngineDir() / TEXT("Build/BatchFiles/Linux/RunMono.sh"));
-	CmdLineParams = FString::Printf(TEXT("\"%s\" \"%s\" %s"), *ScriptPath, *ExecutableFileName, *CmdLineParams);
-	ExecutableFileName = TEXT("/bin/bash");
-#endif
-
 	// Run UnrealBuildTool
 	const bool bLaunchDetached = false;
 	const bool bLaunchHidden = true;
diff --git a/Engine/Source/Developer/DesktopPlatform/Private/Linux/DesktopPlatformLinux.cpp b/Engine/Source/Developer/DesktopPlatform/Private/Linux/DesktopPlatformLinux.cpp
--- a/Engine/Source/Developer/DesktopPlatform/Private/Linux/DesktopPlatformLinux.cpp
+++ b/Engine/Source/Developer/DesktopPlatform/Private/Linux/DesktopPlatformLinux.cpp
@@ -429,19 +429,12 @@ bool FDesktopPlatformLinux::RunUnrealBuildTool(const FText& Description, const F
 	// Write the output
 	Warn->Logf(TEXT("Running %s %s"), *UnrealBuildToolPath, *Arguments);
 
-	// launch UBT with Mono
-	FString ScriptPath = FPaths::ConvertRelativePathToFull(RootDir / TEXT("Engine/Build/BatchFiles/Linux/RunMono.sh"));
-	FString CmdLineParams = FString::Printf(TEXT("\"%s\" \"%s\" %s"), *ScriptPath, *UnrealBuildToolPath, *Arguments);
-
-	// Spawn it with bash (and not sh) because of pushd
-	return FFeedbackContextMarkup::PipeProcessOutput(Description, TEXT("/bin/bash"), CmdLineParams, Warn, &OutExitCode) && OutExitCode == 0;
+	return FFeedbackContextMarkup::PipeProcessOutput(Description, UnrealBuildToolPath, Arguments, Warn, &OutExitCode) && OutExitCode == 0;
 }
 
 bool FDesktopPlatformLinux::IsUnrealBuildToolRunning()
 {
-	// For now assume that if mono application is running, we're running UBT
-	// @todo: we need to get the commandline for the mono process and check if UBT.exe is in there.
-	return FPlatformProcess::IsApplicationRunning(TEXT("mono"));
+	return FPlatformProcess::IsApplicationRunning(TEXT("UnrealBuildTool"));
 }
 
 FFeedbackContext* FDesktopPlatformLinux::GetNativeFeedbackContext()
diff --git a/Engine/Source/Developer/DesktopPlatform/Private/Mac/DesktopPlatformMac.cpp b/Engine/Source/Developer/DesktopPlatform/Private/Mac/DesktopPlatformMac.cpp
--- a/Engine/Source/Developer/DesktopPlatform/Private/Mac/DesktopPlatformMac.cpp
+++ b/Engine/Source/Developer/DesktopPlatform/Private/Mac/DesktopPlatformMac.cpp
@@ -609,19 +609,13 @@ bool FDesktopPlatformMac::RunUnrealBuildTool(const FText& Description, const FSt
 		return false;
 	}
 
-	// On Mac we launch UBT with Mono
-	FString ScriptPath = FPaths::ConvertRelativePathToFull(RootDir / TEXT("Engine/Build/BatchFiles/Mac/RunMono.sh"));
-	FString CmdLineParams = FString::Printf(TEXT("\"%s\" \"%s\" %s"), *ScriptPath, *UnrealBuildToolPath, *Arguments);
-
 	// Spawn it
-	return FFeedbackContextMarkup::PipeProcessOutput(Description, TEXT("/bin/sh"), CmdLineParams, Warn, &OutExitCode) && OutExitCode == 0;
+	return FFeedbackContextMarkup::PipeProcessOutput(Description, UnrealBuildToolPath, Arguments, Warn, &OutExitCode) && OutExitCode == 0;
 }
 
 bool FDesktopPlatformMac::IsUnrealBuildToolRunning()
 {
-	// For now assume that if mono application is running, we're running UBT
-	// @todo: we need to get the commandline for the mono process and check if UBT.exe is in there.
-	return FPlatformProcess::IsApplicationRunning(TEXT("mono"));
+	return FPlatformProcess::IsApplicationRunning(TEXT("UnrealBuildTool"));
 }
 
 FFeedbackContext* FDesktopPlatformMac::GetNativeFeedbackContext()
