Read Engini Ini from the root path, this will, when deployed to the installer branch, fix a issue when attempting to generate projects from the context menu of a uproject.

diff --git a/Engine/Source/Developer/DesktopPlatform/Private/DesktopPlatformBase.cpp b/Engine/Source/Developer/DesktopPlatform/Private/DesktopPlatformBase.cpp
--- a/Engine/Source/Developer/DesktopPlatform/Private/DesktopPlatformBase.cpp
+++ b/Engine/Source/Developer/DesktopPlatform/Private/DesktopPlatformBase.cpp
@@ -1457,12 +1457,17 @@ FString FDesktopPlatformBase::GetUnrealBuildToolProjectFileName(const FString& R
 
 FString FDesktopPlatformBase::GetUnrealBuildToolExecutableFilename(const FString& RootDir) const
 {
-	FString Entry;
-	if( GConfig->GetString( TEXT("PlatformPaths"), TEXT("UnrealBuildTool"), Entry, GEngineIni ))
+	FConfigFile Config;
+	if (FConfigCacheIni::LoadExternalIniFile(Config, TEXT("Engine"), *FPaths::Combine(RootDir, TEXT("Engine/Config/")), *FPaths::Combine(RootDir, TEXT("Engine/Config/")), true, NULL, false, /*bWriteDestIni*/ false))
 	{
-		FString NewPath = FPaths::ConvertRelativePathToFull(RootDir / Entry);
-		return NewPath;
+		FString Entry;
+		if( Config.GetString( TEXT("PlatformPaths"), TEXT("UnrealBuildTool"), Entry ))
+		{
+			FString NewPath = FPaths::ConvertRelativePathToFull(RootDir / Entry);
+			return NewPath;
+		}
 	}
+	
 
 	return FPaths::ConvertRelativePathToFull(RootDir / TEXT("Engine/Binaries/DotNET/UnrealBuildTool.exe"));
 }
