Fix UE4 projects saved in custom locations not being found by the UE5 project picker

diff --git a/Engine/Source/Developer/DesktopPlatform/Private/DesktopPlatformBase.cpp b/Engine/Source/Developer/DesktopPlatform/Private/DesktopPlatformBase.cpp
--- a/Engine/Source/Developer/DesktopPlatform/Private/DesktopPlatformBase.cpp
+++ b/Engine/Source/Developer/DesktopPlatform/Private/DesktopPlatformBase.cpp
@@ -1120,26 +1120,14 @@ void FDesktopPlatformBase::GetProjectBuildProducts(const FString& ProjectDir, TA
 
 FString FDesktopPlatformBase::GetEngineSavedConfigDirectory(const FString& Identifier)
 {
-	// Get the engine root directory
-	FString RootDir;
-	if (!GetEngineRootDirFromIdentifier(Identifier, RootDir))
-	{
-		return FString();
-	}
-
-	// Get the path to the game agnostic settings
-	FString UserDir;
-	if (IsStockEngineRelease(Identifier))
-	{
-		UserDir = FPaths::Combine(FPlatformProcess::UserSettingsDir(), *FApp::GetEpicProductIdentifier(), *Identifier);
-	}
-	else
+	// Get the game agnostic config dir
+	const FString UserDir = GetUserDir(Identifier);
+	if (!UserDir.IsEmpty())
 	{
-		UserDir = FPaths::Combine(*RootDir, TEXT("Engine"));
+		return UserDir / TEXT("Saved/Config") / ANSI_TO_TCHAR(FPlatformProperties::PlatformName());
 	}
 
-	// Get the game agnostic config dir
-	return UserDir / TEXT("Saved/Config") / ANSI_TO_TCHAR(FPlatformProperties::PlatformName());
+	return FString();
 }
 
 bool FDesktopPlatformBase::EnumerateProjectsKnownByEngine(const FString &Identifier, bool bIncludeNativeProjects, TArray<FString> &OutProjectFileNames)
@@ -1152,7 +1140,6 @@ bool FDesktopPlatformBase::EnumerateProjectsKnownByEngine(const FString &Identif
 	}
 
 	FString GameAgnosticConfigDir = GetEngineSavedConfigDirectory(Identifier);
-
 	if (GameAgnosticConfigDir.Len() == 0)
 	{
 		return false;
@@ -1164,7 +1151,12 @@ bool FDesktopPlatformBase::EnumerateProjectsKnownByEngine(const FString &Identif
 
 	// Load the config file
 	FConfigFile GameAgnosticConfig;
-	FConfigCacheIni::LoadExternalIniFile(GameAgnosticConfig, TEXT("EditorSettings"), NULL, *GameAgnosticConfigDir, false);
+	if (!FConfigCacheIni::LoadExternalIniFile(GameAgnosticConfig, TEXT("EditorSettings"), NULL, *GameAgnosticConfigDir, false))
+	{
+		// Load from the legacy path. Most likely a pre-UE5 engine install
+		GameAgnosticConfigDir = GetLegacyEngineSavedConfigDirectory(Identifier);
+		FConfigCacheIni::LoadExternalIniFile(GameAgnosticConfig, TEXT("EditorSettings"), NULL, *GameAgnosticConfigDir, false);
+	}
 
 	// Find the editor game-agnostic settings
 	FConfigSection* Section = GameAgnosticConfig.Find(TEXT("/Script/UnrealEd.EditorSettings"));
@@ -1450,6 +1442,40 @@ bool FDesktopPlatformBase::ReadTargetInfo(const FString& FileName, TArray<FTarge
 	return true;
 }
 
+FString FDesktopPlatformBase::GetUserDir(const FString& Identifier)
+{
+	// Get the engine root directory
+	FString RootDir;
+	if (!GetEngineRootDirFromIdentifier(Identifier, RootDir))
+	{
+		return FString();
+	}
+
+	// Get the path to the game agnostic settings
+	FString UserDir;
+	if (IsStockEngineRelease(Identifier))
+	{
+		UserDir = FPaths::Combine(FPlatformProcess::UserSettingsDir(), *FApp::GetEpicProductIdentifier(), *Identifier);
+	}
+	else
+	{
+		UserDir = FPaths::Combine(*RootDir, TEXT("Engine"));
+	}
+
+	return UserDir;
+}
+
+FString FDesktopPlatformBase::GetLegacyEngineSavedConfigDirectory(const FString& Identifier)
+{
+	const FString UserDir = GetUserDir(Identifier);
+	if (!UserDir.IsEmpty())
+	{
+		return UserDir / TEXT("Saved/Config") / ANSI_TO_TCHAR(FPlatformProperties::IniPlatformName());
+	}
+
+	return FString();
+}
+
 FString FDesktopPlatformBase::GetUnrealBuildToolProjectFileName(const FString& RootDir) const
 {
 	return FPaths::ConvertRelativePathToFull(RootDir / TEXT("Engine/Source/Programs/UnrealBuildTool/UnrealBuildTool.csproj"));
diff --git a/Engine/Source/Developer/DesktopPlatform/Private/DesktopPlatformBase.h b/Engine/Source/Developer/DesktopPlatform/Private/DesktopPlatformBase.h
--- a/Engine/Source/Developer/DesktopPlatform/Private/DesktopPlatformBase.h
+++ b/Engine/Source/Developer/DesktopPlatform/Private/DesktopPlatformBase.h
@@ -31,6 +31,7 @@ public:
 	virtual bool GetDefaultEngineIdentifier(FString &OutIdentifier) override;
 	virtual bool GetDefaultEngineRootDir(FString &OutRootDir) override;
 	virtual FString GetEngineSavedConfigDirectory(const FString& Identifier) override;
+	
 	virtual bool IsPreferredEngineIdentifier(const FString &Identifier, const FString &OtherIdentifier) override;
 
 	virtual bool TryGetEngineVersion(const FString& RootDir, FEngineVersion& OutVersion) override;
@@ -82,6 +83,10 @@ private:
 
 	static bool ReadTargetInfo(const FString& FileName, TArray<FTargetInfo>& Targets);
 
+	FString GetUserDir(const FString& Identifier);
+
+	/** Pre-UE5 location of saved configs */
+	FString GetLegacyEngineSavedConfigDirectory(const FString& Identifier);
 protected:
 
 	FString GetUnrealBuildToolProjectFileName(const FString& RootDir) const;
