Fix dotnet blocking forever if no pipes are passed in

diff --git a/Engine/Source/Developer/DesktopPlatform/Private/DesktopPlatformBase.cpp b/Engine/Source/Developer/DesktopPlatform/Private/DesktopPlatformBase.cpp
--- a/Engine/Source/Developer/DesktopPlatform/Private/DesktopPlatformBase.cpp
+++ b/Engine/Source/Developer/DesktopPlatform/Private/DesktopPlatformBase.cpp
@@ -734,7 +734,7 @@ FProcHandle FDesktopPlatformBase::InvokeUnrealBuildToolAsync(const FString& InCm
 	const bool bLaunchHidden = true;
 	const bool bLaunchReallyHidden = bLaunchHidden;
 
-	FProcHandle ProcHandle = FPlatformProcess::CreateProc(*ExecutableFileName, *CmdLineParams, bLaunchDetached, bLaunchHidden, bLaunchReallyHidden, NULL, 0, NULL, OutWritePipe);
+	FProcHandle ProcHandle = FPlatformProcess::CreateProc(*ExecutableFileName, *CmdLineParams, bLaunchDetached, bLaunchHidden, bLaunchReallyHidden, NULL, 0, NULL, OutWritePipe, OutReadPipe);
 	if (!ProcHandle.IsValid())
 	{
 		Ar.Logf(TEXT("Failed to launch Unreal Build Tool. (%s)"), *ExecutableFileName);
@@ -1366,7 +1366,13 @@ bool FDesktopPlatformBase::BuildUnrealBuildTool(const FString& RootDir, FOutputD
 	const bool bLaunchDetached = false;
 	const bool bLaunchHidden = true;
 	const bool bLaunchReallyHidden = bLaunchHidden;
-	FProcHandle ProcHandle = FPlatformProcess::CreateProc(*CompilerExecutableFilename, *CmdLineParams, bLaunchDetached, bLaunchHidden, bLaunchReallyHidden, NULL, 0, NULL, NULL);
+
+	void* PipeRead = nullptr;
+	void* PipeWrite = nullptr;
+
+	verify(FPlatformProcess::CreatePipe(PipeRead, PipeWrite));
+
+	FProcHandle ProcHandle = FPlatformProcess::CreateProc(*CompilerExecutableFilename, *CmdLineParams, bLaunchDetached, bLaunchHidden, bLaunchReallyHidden, NULL, 0, NULL, PipeWrite, PipeRead);
 	if (!ProcHandle.IsValid())
 	{
 		Ar.Log(TEXT("Failed to start process."));
@@ -1374,6 +1380,7 @@ bool FDesktopPlatformBase::BuildUnrealBuildTool(const FString& RootDir, FOutputD
 	}
 	FPlatformProcess::WaitForProc(ProcHandle);
 	FPlatformProcess::CloseProc(ProcHandle);
+	FPlatformProcess::ClosePipe(PipeRead, PipeWrite);
 
 	// If the executable appeared where we expect it, then we were successful
 	FString UnrealBuildToolExePath = GetUnrealBuildToolExecutableFilename(RootDir);
