Update Icons and rebuild UnrealVersionSelector for Linux

diff --git a/Engine/Source/Developer/DesktopPlatform/Private/Linux/DesktopPlatformLinux.cpp b/Engine/Source/Developer/DesktopPlatform/Private/Linux/DesktopPlatformLinux.cpp
--- a/Engine/Source/Developer/DesktopPlatform/Private/Linux/DesktopPlatformLinux.cpp
+++ b/Engine/Source/Developer/DesktopPlatform/Private/Linux/DesktopPlatformLinux.cpp
@@ -354,7 +354,7 @@ bool FDesktopPlatformLinux::UpdateFileAssociations()
 		return false;
 	}
 
-	if (!RunXDGUtil(FString::Printf(TEXT("xdg-icon-resource install --novendor --mode user --context apps --size 256 %sRuntime/Launch/Resources/Linux/UE4.png ubinary"), *FPaths::EngineSourceDir())))
+	if (!RunXDGUtil(FString::Printf(TEXT("xdg-icon-resource install --novendor --mode user --context apps --size 256 %sRuntime/Launch/Resources/Linux/UnrealEngine.png ubinary"), *FPaths::EngineSourceDir())))
 	{
 		return false;
 	}
