- Added warnings about failure to run UBT
- Fixed (an unused) fallback path for proper path to UBT

diff --git a/Engine/Source/Developer/DesktopPlatform/Private/DesktopPlatformBase.cpp b/Engine/Source/Developer/DesktopPlatform/Private/DesktopPlatformBase.cpp
--- a/Engine/Source/Developer/DesktopPlatform/Private/DesktopPlatformBase.cpp
+++ b/Engine/Source/Developer/DesktopPlatform/Private/DesktopPlatformBase.cpp
@@ -718,6 +718,7 @@ FProcHandle FDesktopPlatformBase::InvokeUnrealBuildToolAsync(const FString& InCm
 			{
 				// Failed to build UBT
 				Ar.Log(TEXT("Failed to build UnrealBuildTool."));
+				UE_LOG(LogDesktopPlatform, Warning, TEXT("Failed to compile UnrealBuildTool (project file is '%s', exe path is '%s')"), *GetUnrealBuildToolProjectFileName(FPaths::RootDir()), *ExecutableFileName);
 				return FProcHandle();
 			}
 		}
@@ -737,6 +738,7 @@ FProcHandle FDesktopPlatformBase::InvokeUnrealBuildToolAsync(const FString& InCm
 	FProcHandle ProcHandle = FPlatformProcess::CreateProc(*ExecutableFileName, *CmdLineParams, bLaunchDetached, bLaunchHidden, bLaunchReallyHidden, NULL, 0, NULL, OutWritePipe, OutReadPipe);
 	if (!ProcHandle.IsValid())
 	{
+		UE_LOG(LogDesktopPlatform, Warning, TEXT("Failed to launch UnrealBuildTool (exe path is '%s')"), *ExecutableFileName);
 		Ar.Logf(TEXT("Failed to launch Unreal Build Tool. (%s)"), *ExecutableFileName);
 	}
 
@@ -1502,7 +1504,7 @@ FString FDesktopPlatformBase::GetUnrealBuildToolExecutableFilename(const FString
 	}
 	
 
-	return FPaths::ConvertRelativePathToFull(RootDir / TEXT("Engine/Binaries/DotNET/UnrealBuildTool.exe"));
+	return FPaths::ConvertRelativePathToFull(RootDir / TEXT("Engine/Binaries/DotNET/UnrealBuildTool/UnrealBuildTool.exe"));
 }
 
 #undef LOCTEXT_NAMESPACE
