Added some temp logging to track down missing projects in the project browser from the binary builds. Super verbose but should only spam when opening the project browser.

diff --git a/Engine/Source/Developer/DesktopPlatform/Private/DesktopPlatformBase.cpp b/Engine/Source/Developer/DesktopPlatform/Private/DesktopPlatformBase.cpp
--- a/Engine/Source/Developer/DesktopPlatform/Private/DesktopPlatformBase.cpp
+++ b/Engine/Source/Developer/DesktopPlatform/Private/DesktopPlatformBase.cpp
@@ -1151,12 +1151,20 @@ bool FDesktopPlatformBase::EnumerateProjectsKnownByEngine(const FString &Identif
 	TArray<FString> SearchDirectories;
 	SearchDirectories.AddUnique(GetDefaultProjectCreationPath());
 
+	UE_LOG(LogDesktopPlatform, Log, TEXT("Enumerating Projects From Engine Ver: %s"), *Identifier);
+
+	UE_LOG(LogDesktopPlatform, Log, TEXT("Looking for directories to scan from : %s"), *GameAgnosticConfigDir);
 	// Load the config file
 	FConfigFile GameAgnosticConfig;
 	if (!FConfigCacheIni::LoadExternalIniFile(GameAgnosticConfig, TEXT("EditorSettings"), NULL, *GameAgnosticConfigDir, false))
 	{
+		FString PreviousConfigDir = MoveTemp(GameAgnosticConfigDir);
+
 		// Load from the legacy path. Most likely a pre-UE5 engine install
 		GameAgnosticConfigDir = GetLegacyEngineSavedConfigDirectory(Identifier);
+
+		UE_LOG(LogDesktopPlatform, Log, TEXT("%s not found, looking for directories in %s"), *PreviousConfigDir , *GameAgnosticConfigDir);
+
 		FConfigCacheIni::LoadExternalIniFile(GameAgnosticConfig, TEXT("EditorSettings"), NULL, *GameAgnosticConfigDir, false);
 	}
 
@@ -1169,14 +1177,23 @@ bool FDesktopPlatformBase::EnumerateProjectsKnownByEngine(const FString &Identif
 		Section = GameAgnosticConfig.Find(TEXT("/Script/UnrealEd.EditorGameAgnosticSettings"));
 	}
 
+	if (GameAgnosticConfig.IsEmpty())
+	{
+		UE_LOG(LogDesktopPlatform, Log, TEXT("Config not found"));
+	}
+
 	if(Section != NULL)
 	{
+		UE_LOG(LogDesktopPlatform, Log, TEXT("Searching for previously created project directories..."));
+
 		// Add in every path that the user has ever created a project file. This is to catch new projects showing up in the user's project folders
 		TArray<FString> AdditionalDirectories;
 		Section->MultiFind(TEXT("CreatedProjectPaths"), AdditionalDirectories);
 		for(int Idx = 0; Idx < AdditionalDirectories.Num(); Idx++)
 		{
 			FPaths::NormalizeDirectoryName(AdditionalDirectories[Idx]);
+
+			UE_LOG(LogDesktopPlatform, Log, TEXT("Found directory: \"%s\""), *AdditionalDirectories[Idx]);
 			SearchDirectories.AddUnique(AdditionalDirectories[Idx]);
 		}
 
@@ -1186,6 +1203,9 @@ bool FDesktopPlatformBase::EnumerateProjectsKnownByEngine(const FString &Identif
 		for(int Idx = 0; Idx < RecentlyOpenedFiles.Num(); Idx++)
 		{
 			FPaths::NormalizeFilename(RecentlyOpenedFiles[Idx]);
+
+			UE_LOG(LogDesktopPlatform, Log, TEXT("Found project \"%s\" in recently opened files"), *RecentlyOpenedFiles[Idx]);
+
 			OutProjectFileNames.AddUnique(RecentlyOpenedFiles[Idx]);
 		}		
 	}
@@ -1203,11 +1223,17 @@ bool FDesktopPlatformBase::EnumerateProjectsKnownByEngine(const FString &Identif
 
 			for(int32 FileIdx = 0; FileIdx < ProjectFiles.Num(); FileIdx++)
 			{
-				OutProjectFileNames.AddUnique(SearchDirectories[Idx] / ProjectFolders[FolderIdx] / ProjectFiles[FileIdx]);
+				FString ProjName = SearchDirectories[Idx] / ProjectFolders[FolderIdx] / ProjectFiles[FileIdx];
+
+				UE_LOG(LogDesktopPlatform, Log, TEXT("Found project \"%s\" in previously created project directory"), *ProjName);
+
+				OutProjectFileNames.AddUnique(MoveTemp(ProjName));
 			}
 		}
 	}
 
+	UE_LOG(LogDesktopPlatform, Log, TEXT("Searcing for projects in .uprojectdirs"));
+
 	// Find all the native projects, and either add or remove them from the list depending on whether we want native projects
 	const FUProjectDictionary &Dictionary = GetCachedProjectDictionary(RootDir);
 	if(bIncludeNativeProjects)
@@ -1217,6 +1243,8 @@ bool FDesktopPlatformBase::EnumerateProjectsKnownByEngine(const FString &Identif
 		{
 			if(!NativeProjectPaths[Idx].Contains(TEXT("/Templates/")))
 			{
+				UE_LOG(LogDesktopPlatform, Log, TEXT("Found project \"%s\" in .uprojectdirs"), *NativeProjectPaths[Idx]);
+
 				OutProjectFileNames.AddUnique(NativeProjectPaths[Idx]);
 			}
 		}
diff --git a/Engine/Source/Editor/GameProjectGeneration/Private/SProjectBrowser.cpp b/Engine/Source/Editor/GameProjectGeneration/Private/SProjectBrowser.cpp
--- a/Engine/Source/Editor/GameProjectGeneration/Private/SProjectBrowser.cpp
+++ b/Engine/Source/Editor/GameProjectGeneration/Private/SProjectBrowser.cpp
@@ -54,6 +54,8 @@
 
 #define LOCTEXT_NAMESPACE "ProjectBrowser"
 
+DEFINE_LOG_CATEGORY_STATIC(LogProjectBrowser, Log, All);
+
 namespace ProjectBrowserDefs
 {
 	constexpr float ProjectTileHeight = 153;
@@ -661,11 +663,15 @@ FReply SProjectBrowser::FindProjects()
 	TMap<FString, FString> EngineInstallations;
 	FDesktopPlatformModule::Get()->EnumerateEngineInstallations(EngineInstallations);
 
+	UE_LOG(LogProjectBrowser, Log, TEXT("Looking for projects..."));
+
 	// Add projects from every branch that we know about
 	for (TMap<FString, FString>::TConstIterator Iter(EngineInstallations); Iter; ++Iter)
 	{
 		TArray<FString> ProjectFiles;
 
+		UE_LOG(LogProjectBrowser, Log, TEXT("Found Engine Installation \"%s\"(%s)"), *Iter.Key(), *Iter.Value());
+
 		if (FDesktopPlatformModule::Get()->EnumerateProjectsKnownByEngine(Iter.Key(), false, ProjectFiles))
 		{
 			AllFoundProjectFiles.Append(MoveTemp(ProjectFiles));
@@ -675,6 +681,14 @@ FReply SProjectBrowser::FindProjects()
 	// Add all the samples from the launcher
 	TArray<FString> LauncherSampleProjects;
 	FDesktopPlatformModule::Get()->EnumerateLauncherSampleProjects(LauncherSampleProjects);
+
+	UE_LOG(LogProjectBrowser, Log, TEXT("Enumerating Launcher Sample Projects..."));
+
+	for (const FString& Str : LauncherSampleProjects)
+	{
+		UE_LOG(LogProjectBrowser, Log, TEXT("Found Sample Project:\"%s\""), *Str);
+	}
+
 	AllFoundProjectFiles.Append(MoveTemp(LauncherSampleProjects));
 
 	// Add all the native project files we can find
@@ -713,7 +727,7 @@ FReply SProjectBrowser::FindProjects()
 		}
 	}
 
-	ProjectItemsSource.Sort([](auto& A, auto& B) { return A->LastAccessTime > B->LastAccessTime; });
+	ProjectItemsSource.Sort([](auto& A, auto& B) { return A->LastAccessTime < B->LastAccessTime; });
 
 	PopulateFilteredProjects();
 
