macOS file open filter improvements

This changes the way filters for the open dialog are done to be more flexible with existing filter text.

#rnx

diff --git a/Engine/Source/Developer/DesktopPlatform/Private/Mac/DesktopPlatformMac.cpp b/Engine/Source/Developer/DesktopPlatform/Private/Mac/DesktopPlatformMac.cpp
--- a/Engine/Source/Developer/DesktopPlatform/Private/Mac/DesktopPlatformMac.cpp
+++ b/Engine/Source/Developer/DesktopPlatform/Private/Mac/DesktopPlatformMac.cpp
@@ -12,6 +12,71 @@
 
 #define LOCTEXT_NAMESPACE "DesktopPlatform"
 
+// Utility function that will take a string and parse out all extensions given a separator.
+// Extensions are trimmed for whitespace and made lowercase.
+// ie:
+// ExtensionsInStringSeparatedByString( @"*.foo, *.bar,  *.moof", @"," ) = @[ @"foo", @"bar", @"moof" ]
+
+NSArray<NSString*>* ExtensionsInStringSeparatedByString(NSString* CombinedString, NSString* Separator)
+{
+	NSMutableSet<NSString*>* AllowedFileExtensions = [NSMutableSet new];
+	
+	NSArray<NSString*>* Filters = [CombinedString componentsSeparatedByString:Separator];
+	
+	for (NSString* Filter in Filters)
+	{
+		NSString* TrimmedFilter = [Filter stringByTrimmingCharactersInSet:[NSCharacterSet whitespaceCharacterSet]];
+		if ( [TrimmedFilter hasPrefix:@"*."] )
+		{
+			NSString* Ext = [[TrimmedFilter substringFromIndex:2] lowercaseString];
+			if ( Ext.length > 0 )
+			{
+				[AllowedFileExtensions addObject:Ext];
+			}
+		}
+	}
+	
+	if ( AllowedFileExtensions.count == 0 )
+	{
+		return nil;
+	}
+	return [AllowedFileExtensions allObjects];
+}
+
+// This will get an array of extensions from a filter text string.  Supported formats examples are:
+// Image File (*.png,*.jpg)|*.png,*.jpg
+// Image File|*.png,*.jpg
+// Image File (*.png,*.jpg)
+//
+// The result array extensions are trimmed for whitespace and made lowercase.
+
+NSArray<NSString*>* AllowedFileExtensionsForFilterText(NSString* FilterText)
+{
+	NSArray<NSString*>* PipeSplit = [FilterText componentsSeparatedByString:@"|"];
+	if ( PipeSplit.count > 1 )
+	{
+		// We only need the text after the first pipe here
+		return ExtensionsInStringSeparatedByString( PipeSplit[1], @";" );
+	}
+	
+	
+	NSArray<NSString*>* OpenBracketSplit = [FilterText componentsSeparatedByString:@"("];
+	if ( OpenBracketSplit.count != 2 )
+	{
+		return nil;
+	}
+
+	NSString* AfterOpenBracket = OpenBracketSplit[1];
+	if ( ![AfterOpenBracket hasSuffix:@")"] )
+	{
+		return nil;
+	}
+	
+	NSString* FilterContentText = [AfterOpenBracket substringWithRange:NSMakeRange(0, AfterOpenBracket.length-1)];
+	
+	return ExtensionsInStringSeparatedByString( FilterContentText, @"," );
+}
+
 class FMacScopedSystemModalMode
 {
 public:
@@ -140,26 +205,18 @@ private:
 
 - (void)SetExtensionsAtIndex: (int32) index
 {
-	check( [AllowedFileTypes count] >= index * 2 );
+	check( [AllowedFileTypes count] >= index );
 	SelectedExtension = index;
 
-	NSString* ExtsToParse = [AllowedFileTypes objectAtIndex:index * 2 + 1];
+	NSString* ExtsToParse = [AllowedFileTypes objectAtIndex:index];
 	if( [ExtsToParse compare:@"*.*"] == NSOrderedSame )
 	{
 		[DialogPanel setAllowedFileTypes: nil];
 	}
 	else
 	{
-		NSArray* ExtensionsWildcards = [ExtsToParse componentsSeparatedByString:@";"];
-		NSMutableArray* Extensions = [NSMutableArray arrayWithCapacity: [ExtensionsWildcards count]];
-
-		for( int32 Index = 0; Index < [ExtensionsWildcards count]; ++Index )
-		{
-			NSString* Temp = [[ExtensionsWildcards objectAtIndex:Index] stringByTrimmingCharactersInSet:[NSCharacterSet characterSetWithCharactersInString:@"*."]];
-			[Extensions addObject: Temp];
-		}
-
-		[DialogPanel setAllowedFileTypes: Extensions];
+		NSArray<NSString*>* Types = AllowedFileExtensionsForFilterText(ExtsToParse);
+		[DialogPanel setAllowedFileTypes: Types];
 	}
 }
 
diff --git a/Engine/Source/Editor/StringTableEditor/Private/StringTableEditor.cpp b/Engine/Source/Editor/StringTableEditor/Private/StringTableEditor.cpp
--- a/Engine/Source/Editor/StringTableEditor/Private/StringTableEditor.cpp
+++ b/Engine/Source/Editor/StringTableEditor/Private/StringTableEditor.cpp
@@ -662,7 +662,7 @@ FReply FStringTableEditor::OnImportFromCSVClicked()
 			LOCTEXT("ImportStringTableTitle", "Choose a string table CSV file...").ToString(),
 			DefaultPath,
 			TEXT(""),
-			TEXT("String Table CSV (*.csv, *.csv)"),
+			TEXT("String Table CSV (*.csv)"),
 			EFileDialogFlags::None,
 			OutFiles
 			))
