Add build worker discovery mechanism using target receipt files.

diff --git a/Engine/Source/Developer/DesktopPlatform/DesktopPlatform.Build.cs b/Engine/Source/Developer/DesktopPlatform/DesktopPlatform.Build.cs
--- a/Engine/Source/Developer/DesktopPlatform/DesktopPlatform.Build.cs
+++ b/Engine/Source/Developer/DesktopPlatform/DesktopPlatform.Build.cs
@@ -8,6 +8,8 @@ public class DesktopPlatform : ModuleRules
 	{
 		PrivateIncludePaths.Add("Developer/DesktopPlatform/Private");
 
+		PrivateIncludePathModuleNames.Add("DerivedDataCache");
+
 		PrivateDependencyModuleNames.AddRange(
 			new string[] {
 				"Core",
diff --git a/Engine/Source/Developer/DesktopPlatform/Private/MultiPlatformTargetReceiptBuildWorkers.cpp b/Engine/Source/Developer/DesktopPlatform/Private/MultiPlatformTargetReceiptBuildWorkers.cpp
new file mode 100644
--- /dev/null
+++ b/Engine/Source/Developer/DesktopPlatform/Private/MultiPlatformTargetReceiptBuildWorkers.cpp
@@ -0,0 +1,33 @@
+// Copyright Epic Games, Inc. All Rights Reserved.
+
+#include "MultiPlatformTargetReceiptBuildWorkers.h"
+
+FMultiPlatformTargetReceiptBuildWorkers::FMultiPlatformTargetReceiptBuildWorkers(const TCHAR* TargetReceiptFilePath)
+: bAllPlatformsInitialized(false)
+{
+	FString PathString = TargetReceiptFilePath;
+	if (!PathString.Contains(TEXT("$(Platform)")))
+	{
+		new(&PlatformSpecificWorkerFactories[0]) FTargetReceiptBuildWorker(*PathString);
+		return;
+	}
+
+	new(&PlatformSpecificWorkerFactories[SupportedPlatform_Win64]) FTargetReceiptBuildWorker(*PathString.Replace(TEXT("$(Platform)"), TEXT("Win64")));
+	new(&PlatformSpecificWorkerFactories[SupportedPlatform_Mac]) FTargetReceiptBuildWorker(*PathString.Replace(TEXT("$(Platform)"), TEXT("Mac")));
+	new(&PlatformSpecificWorkerFactories[SupportedPlatform_Linux]) FTargetReceiptBuildWorker(*PathString.Replace(TEXT("$(Platform)"), TEXT("Linux")));
+	bAllPlatformsInitialized = true;
+}
+
+FMultiPlatformTargetReceiptBuildWorkers::~FMultiPlatformTargetReceiptBuildWorkers()
+{
+	if (bAllPlatformsInitialized)
+	{
+		((FTargetReceiptBuildWorker*)&PlatformSpecificWorkerFactories[SupportedPlatform_Win64])->~FTargetReceiptBuildWorker();
+		((FTargetReceiptBuildWorker*)&PlatformSpecificWorkerFactories[SupportedPlatform_Mac])->~FTargetReceiptBuildWorker();
+		((FTargetReceiptBuildWorker*)&PlatformSpecificWorkerFactories[SupportedPlatform_Linux])->~FTargetReceiptBuildWorker();
+	}
+	else
+	{
+		((FTargetReceiptBuildWorker*)&PlatformSpecificWorkerFactories[0])->~FTargetReceiptBuildWorker();
+	}
+}
diff --git a/Engine/Source/Developer/DesktopPlatform/Private/TargetReceipt.cpp b/Engine/Source/Developer/DesktopPlatform/Private/TargetReceipt.cpp
--- a/Engine/Source/Developer/DesktopPlatform/Private/TargetReceipt.cpp
+++ b/Engine/Source/Developer/DesktopPlatform/Private/TargetReceipt.cpp
@@ -22,7 +22,7 @@ bool TryParseBuildVersion(const FJsonObject& Object, FBuildVersion& OutVersion)
 	return false;
 }
 
-bool FTargetReceipt::Read(const FString& FileName)
+bool FTargetReceipt::Read(const FString& FileName, bool bExpandVariables)
 {
 	// Read the file from disk
 	FString FileContents;
@@ -87,7 +87,10 @@ bool FTargetReceipt::Read(const FString& FileName)
 	{
 		return false;
 	}
-	ExpandVariables(Launch);
+	if (bExpandVariables)
+	{
+		ExpandVariables(Launch);
+	}
 
 	// Read the list of build products
 	const TArray<TSharedPtr<FJsonValue>>* BuildProductsArray;
@@ -107,7 +110,10 @@ bool FTargetReceipt::Read(const FString& FileName)
 				return false;
 			}
 
-			ExpandVariables(BuildProduct.Path);
+			if (bExpandVariables)
+			{
+				ExpandVariables(BuildProduct.Path);
+			}
 
 			BuildProducts.Add(MoveTemp(BuildProduct));
 		}
@@ -131,7 +137,10 @@ bool FTargetReceipt::Read(const FString& FileName)
 				return false;
 			}
 
-			ExpandVariables(RuntimeDependency.Path);
+			if (bExpandVariables)
+			{
+				ExpandVariables(RuntimeDependency.Path);
+			}
 
 			RuntimeDependencies.Add(MoveTemp(RuntimeDependency));
 		}
@@ -178,7 +187,10 @@ bool FTargetReceipt::Read(const FString& FileName)
 				return false;
 			}
 
-			ExpandVariables(Property.Value);
+			if (bExpandVariables)
+			{
+				ExpandVariables(Property.Value);
+			}
 
 			AdditionalProperties.Add(MoveTemp(Property));
 		}
diff --git a/Engine/Source/Developer/DesktopPlatform/Private/TargetReceiptBuildWorker.cpp b/Engine/Source/Developer/DesktopPlatform/Private/TargetReceiptBuildWorker.cpp
new file mode 100644
--- /dev/null
+++ b/Engine/Source/Developer/DesktopPlatform/Private/TargetReceiptBuildWorker.cpp
@@ -0,0 +1,486 @@
+// Copyright Epic Games, Inc. All Rights Reserved.
+
+#include "TargetReceiptBuildWorker.h"
+
+#include "DerivedDataBuildWorker.h"
+#include "Features/IModularFeatures.h"
+#include "HAL/FileManager.h"
+#include "HAL/Platform.h"
+#include "HAL/PlatformProcess.h"
+#include "HAL/PlatformTime.h"
+#include "Misc/Paths.h"
+#include "Misc/PathViews.h"
+#include "Misc/StringBuilder.h"
+#include "Serialization/CompactBinary.h"
+#include "Serialization/CompactBinarySerialization.h"
+#include "TargetReceipt.h"
+
+DEFINE_LOG_CATEGORY_STATIC(LogTargetReceiptBuildWorker, Log, All);
+
+// This needs to match the version exported in DerivedDataBuildWorker.Build.cs 
+const FGuid FTargetReceiptBuildWorker::WorkerReceiptVersion(TEXT("a547b2fa-e2f9-47c6-80c8-df688fb99f34"));
+
+FTargetReceiptBuildWorker::FWorkerPath::FWorkerPath(FStringView InLocalPathRoot, FStringView InLocalPathSuffix, FStringView InRemotePathRoot, FStringView InRemotePathSuffix)
+{
+	TStringBuilder<128> LocalPathBuilder;
+	FPathViews::Append(LocalPathBuilder, InLocalPathRoot, InLocalPathSuffix);
+	LocalPath = LocalPathBuilder.ToString();
+
+	TStringBuilder<128> RemotePathBuilder;
+	FPathViews::Append(RemotePathBuilder,
+						InRemotePathRoot.IsEmpty() ? InLocalPathRoot : InRemotePathRoot,
+						InRemotePathSuffix.IsEmpty() ? InLocalPathSuffix : InRemotePathSuffix);
+	RemotePath = RemotePathBuilder.ToString();
+}
+
+class FPopulateWorkerFromTargetReceiptTask
+{
+public:
+	FPopulateWorkerFromTargetReceiptTask(const TCHAR* InTargetReceiptFilePath, FTargetReceiptBuildWorker* InWorker)
+	: TargetReceiptFilePath(InTargetReceiptFilePath)
+	, Worker(InWorker)
+	{
+	}
+
+	FORCEINLINE TStatId GetStatId() const { RETURN_QUICK_DECLARE_CYCLE_STAT(FPopulateWorkerFromTargetReceiptTask, STATGROUP_TaskGraphTasks); }
+	ENamedThreads::Type GetDesiredThread() { return ENamedThreads::Type::AnyThread; }
+	static ESubsequentsMode::Type GetSubsequentsMode() { return ESubsequentsMode::FireAndForget; }
+
+	void DoTask(ENamedThreads::Type CurrentThread, const FGraphEventRef& MyCompletionGraphEvent)
+	{
+		Worker->PopulateWorkerFromReceipt(TargetReceiptFilePath);
+	}
+
+private:
+	const FString TargetReceiptFilePath;
+	FTargetReceiptBuildWorker* Worker;
+};
+
+class FTargetReceiptBuildWorkerFactory : public UE::DerivedData::IBuildWorkerFactory
+{
+public:
+	const FTargetReceiptBuildWorker* GetWorker()
+	{
+		return (FTargetReceiptBuildWorker*)((uintptr_t)this - offsetof(class FTargetReceiptBuildWorker, InternalFactory));
+	}
+
+	void Build(UE::DerivedData::FBuildWorkerBuilder& Builder)
+	{
+		const FTargetReceiptBuildWorker* Worker = GetWorker();
+		Builder.SetName(Worker->Name);
+		Builder.SetPath(Worker->ExecutablePaths[0].GetRemotePath());
+		Builder.SetHostPlatform(Worker->Platform);
+
+		for (int32 ExecutableIndex = 0; ExecutableIndex < Worker->ExecutablePaths.Num(); ++ExecutableIndex)
+		{
+			Builder.AddExecutable(Worker->ExecutablePaths[ExecutableIndex].GetRemotePath(), Worker->ExecutableMeta[ExecutableIndex].Key, Worker->ExecutableMeta[ExecutableIndex].Value);
+		}
+
+		for (const TPair<FString, FString>& EnvironmentVariable : Worker->EnvironmentVariables)
+		{
+			Builder.SetEnvironment(EnvironmentVariable.Key, EnvironmentVariable.Value);
+		}
+
+		Builder.SetBuildSystemVersion(Worker->BuildSystemVersion);
+
+		for (const TPair<FString, FGuid>& FunctionVersion : Worker->FunctionVersions)
+		{
+			Builder.AddFunction(FunctionVersion.Key, FunctionVersion.Value);
+		}
+	}
+
+	UE::DerivedData::FRequest FindFileData(TConstArrayView<FIoHash> RawHashes, UE::DerivedData::EPriority Priority, UE::DerivedData::FOnBuildWorkerFileDataComplete&& OnComplete)
+	{
+		if (OnComplete)
+		{
+			const FTargetReceiptBuildWorker* Worker = GetWorker();
+			TArray<FCompressedBuffer> FileDataBuffers;
+			UE::DerivedData::FBuildWorkerFileDataCompleteParams CompleteParams;
+			for (const FIoHash& RawHash : RawHashes)
+			{
+				for (int32 ExecutableIndex = 0; ExecutableIndex < Worker->ExecutablePaths.Num(); ++ExecutableIndex)
+				{
+					const FTargetReceiptBuildWorker::FWorkerPathMeta& Meta = Worker->ExecutableMeta[ExecutableIndex];
+					if (Meta.Key == RawHash)
+					{
+						if (TUniquePtr<FArchive> Ar{IFileManager::Get().CreateFileReader(*Worker->ExecutablePaths[ExecutableIndex].GetLocalPath(), FILEREAD_Silent)})
+						{
+							const int64 TotalSize = Ar->TotalSize();
+							FUniqueBuffer MutableBuffer = FUniqueBuffer::Alloc(uint64(TotalSize));
+							Ar->Serialize(MutableBuffer.GetData(), TotalSize);
+							if (Ar->Close())
+							{
+								FileDataBuffers.Emplace(FCompressedBuffer::Compress(NAME_None, MutableBuffer.MoveToShared()));
+							}
+						}
+					}
+				}
+			}
+
+			if (FileDataBuffers.Num() == RawHashes.Num())
+			{
+				CompleteParams.Status = UE::DerivedData::EStatus::Ok;
+				CompleteParams.Files = FileDataBuffers;
+				OnComplete(MoveTemp(CompleteParams));
+			}
+			else
+			{
+				CompleteParams.Status = UE::DerivedData::EStatus::Error;
+				OnComplete(MoveTemp(CompleteParams));
+			}
+		}
+
+		return UE::DerivedData::FRequest();
+	}
+};
+
+FTargetReceiptBuildWorker::FTargetReceiptBuildWorker(const TCHAR* TargetReceiptFilePath)
+: bAbortRequested(false)
+, bEnabled(false)
+{
+	static_assert(sizeof(FTargetReceiptBuildWorkerFactory) == sizeof(IPureVirtual), "FTargetReceiptBuildWorkerFactory type must always compile to something equivalent to an IPureVirtual size.");
+	FTargetReceiptBuildWorkerFactory* WorkerFactory = GetWorkerFactory();
+	new(WorkerFactory) FTargetReceiptBuildWorkerFactory();
+	PopulateTaskRef = TGraphTask<FPopulateWorkerFromTargetReceiptTask>::CreateTask().ConstructAndDispatchWhenReady(TargetReceiptFilePath, this);
+}
+
+FTargetReceiptBuildWorker::~FTargetReceiptBuildWorker()
+{
+	bAbortRequested = true;
+	if (PopulateTaskRef.IsValid())
+	{
+		PopulateTaskRef->Wait();
+	}
+
+	FTargetReceiptBuildWorkerFactory* WorkerFactory = GetWorkerFactory();
+	if (bEnabled)
+	{
+		IModularFeatures::Get().UnregisterModularFeature(UE::DerivedData::IBuildWorkerFactory::GetFeatureName(), WorkerFactory);
+	}
+	WorkerFactory->~FTargetReceiptBuildWorkerFactory();
+}
+
+FTargetReceiptBuildWorkerFactory* FTargetReceiptBuildWorker::GetWorkerFactory()
+{
+	return (FTargetReceiptBuildWorkerFactory*)&InternalFactory;
+}
+
+bool FTargetReceiptBuildWorker::TryAddExecutablePath(FStringView Path)
+{
+	if (Path.IsEmpty())
+	{
+		UE_LOG(LogTargetReceiptBuildWorker, Warning, TEXT("Empty executable path encountered while populating worker '%s'."), *Name, Path.Len(), Path.GetData());
+		return false;
+	}
+
+	if (!Path.StartsWith(TEXT("$(")))
+	{
+		UE_LOG(LogTargetReceiptBuildWorker, Warning, TEXT("Absolute or unprefixed executable path variable encountered while populating worker '%s': '%.*s'"), *Name, Path.Len(), Path.GetData());
+		return false;
+	}
+
+	int32 EndVariableIndex = Path.Find(TEXT(")"));
+	if (EndVariableIndex == INDEX_NONE)
+	{
+		UE_LOG(LogTargetReceiptBuildWorker, Warning, TEXT("Malformed executable path variable encountered while populating worker '%s': '%.*s'"), *Name, Path.Len(), Path.GetData());
+		return false;
+	}
+
+	FStringView PrefixVariable = Path.Mid(2, EndVariableIndex - 2);
+	if (PrefixVariable == TEXT("EngineDir"))
+	{
+		ExecutablePaths.Emplace(FPaths::EngineDir(), Path.RightChop(EndVariableIndex + 2), TEXT("Engine"));
+	}
+	else if (PrefixVariable == TEXT("ProjectDir"))
+	{
+		ExecutablePaths.Emplace(FPaths::ProjectDir(), Path.RightChop(EndVariableIndex + 2), TEXT("Project"));
+	}
+	else
+	{
+		FString TempVariableName(PrefixVariable);
+		FString EnvironmentVariableValue = FPlatformMisc::GetEnvironmentVariable(*TempVariableName);
+		if (EnvironmentVariableValue.IsEmpty())
+		{
+			UE_LOG(LogTargetReceiptBuildWorker, Warning, TEXT("Unknown executable path variable encountered while populating worker '%s': '%s'"), *Name, *TempVariableName);
+			return false;
+		}
+		ExecutablePaths.Emplace(EnvironmentVariableValue, Path.RightChop(EndVariableIndex + 2), TempVariableName);
+		// By having a dependency that starts with an environment variable, we implicitly create an environment variable to be remoted.
+		// The variable will have the same name as the variable we used, and its value will be a path relative to the platform specific binaries directory with the same name as the variable.
+		EnvironmentVariables.FindOrAdd(TempVariableName, FString(TEXT("..\\..\\..\\")) + TempVariableName);
+	}
+
+	return true;
+}
+
+void FTargetReceiptBuildWorker::PopulateWorkerFromReceipt(const FString& TargetReceiptFilePath)
+{
+	if (bAbortRequested)
+	{
+		return;
+	}
+
+	IFileManager& FileManager = IFileManager::Get();
+	FString FinalTargetReceiptFilePath;
+	FString FinalPublishedVersionFilePath;
+	FString FinalGeneratedVersionFilePath;
+	if (TargetReceiptFilePath.StartsWith(TEXT("$(ProjectDir)")))
+	{
+		FinalTargetReceiptFilePath = FPaths::ConvertRelativePathToFull(TargetReceiptFilePath.Replace(TEXT("$(ProjectDir)"), *FPaths::ProjectDir()));
+		FinalPublishedVersionFilePath = FinalTargetReceiptFilePath + TEXT(".BuildVersion");
+		FinalGeneratedVersionFilePath = FPaths::ConvertRelativePathToFull(TargetReceiptFilePath.Replace(TEXT("$(ProjectDir)"), *FPaths::Combine(FPaths::ProjectSavedDir(), TEXT("GeneratedWorkerBuildVersions"))) + TEXT(".BuildVersion"));
+	}
+	else if (TargetReceiptFilePath.StartsWith(TEXT("$(EngineDir)")))
+	{
+		// Target receipts declared to be in the engine directory may end up in the project directory if they were compiled
+		// with a uproject file on the commandline.  For this reason we allow for the possibility that the target receipt path
+		// can be tried with both the true engine directory or the engine directory treated as the project directory.
+		// Pick whichever is newer (more recently modified).
+		FString ReceiptEnginePath = TargetReceiptFilePath.Replace(TEXT("$(EngineDir)"), *FPaths::EngineDir());
+		FString ReceiptProjectPath = TargetReceiptFilePath.Replace(TEXT("$(EngineDir)"), *FPaths::ProjectDir());
+		FDateTime ReceiptEngineTimeStamp;
+		FDateTime ReceiptProjectTimeStamp;
+		FileManager.GetTimeStampPair(*ReceiptEnginePath, *ReceiptProjectPath, ReceiptEngineTimeStamp, ReceiptProjectTimeStamp);
+		if (ReceiptProjectTimeStamp > ReceiptEngineTimeStamp)
+		{
+			FinalTargetReceiptFilePath = FPaths::ConvertRelativePathToFull(ReceiptProjectPath);
+			FinalPublishedVersionFilePath = FinalTargetReceiptFilePath + TEXT(".BuildVersion");
+			FinalGeneratedVersionFilePath = FPaths::ConvertRelativePathToFull(TargetReceiptFilePath.Replace(TEXT("$(EngineDir)"), *FPaths::Combine(FPaths::ProjectSavedDir(), TEXT("GeneratedWorkerBuildVersions"))) + TEXT(".BuildVersion"));
+		}
+		else
+		{
+			FinalTargetReceiptFilePath = FPaths::ConvertRelativePathToFull(ReceiptEnginePath);
+			FinalPublishedVersionFilePath = FinalTargetReceiptFilePath + TEXT(".BuildVersion");
+			FinalGeneratedVersionFilePath = FPaths::ConvertRelativePathToFull(TargetReceiptFilePath.Replace(TEXT("$(EngineDir)"), *FPaths::Combine(FPaths::EngineSavedDir(), TEXT("GeneratedWorkerBuildVersions"))) + TEXT(".BuildVersion"));
+		}
+	}
+	else
+	{
+		UE_LOG(LogTargetReceiptBuildWorker, Warning, TEXT("Unknown prefix for target receipt path: '%s'"), *TargetReceiptFilePath);
+		return;
+	}
+	
+	if (!FPaths::FileExists(FinalTargetReceiptFilePath))
+	{
+		// This case should return without emitting a warning.
+		return;
+	}
+
+	FTargetReceipt TargetReceipt;
+	if (!TargetReceipt.Read(FinalTargetReceiptFilePath, false))
+	{
+		UE_LOG(LogTargetReceiptBuildWorker, Warning, TEXT("Target receipt failed to parse: '%s'"), *TargetReceiptFilePath);
+		return;
+	}
+
+	if (bAbortRequested)
+	{
+		return;
+	}
+
+	Name = TargetReceipt.TargetName;
+	Platform = TargetReceipt.Platform;
+	if (!TryAddExecutablePath(TargetReceipt.Launch))
+	{
+		return;
+	}
+
+	FGuid ReceiptInlineBuildSystemVersion;
+	TArray<TPair<FString, FGuid>> ReceiptInlineFunctionVersions;
+	bool bUsableBuildWorker = false;
+	for (const FReceiptProperty& Property : TargetReceipt.AdditionalProperties)
+	{
+		if (Property.Name == TEXT("DerivedDataBuildWorkerReceiptVersion"))
+		{
+			bUsableBuildWorker = FGuid(Property.Value) == WorkerReceiptVersion;
+		}
+		else if (Property.Name == TEXT("DerivedDataBuildWorkerBuildSystemVersion"))
+		{
+			if (!FGuid::Parse(Property.Value, ReceiptInlineBuildSystemVersion))
+			{
+				UE_LOG(LogTargetReceiptBuildWorker, Warning, TEXT("Target receipt contains inline build system version guid that could not be parsed: '%s'"), *Property.Value);
+				return;
+			}
+		}
+		else if (Property.Name == TEXT("DerivedDataBuildWorkerFunctionVersion"))
+		{
+			TArray<FString> FunctionComponents;
+			if (2 != Property.Value.ParseIntoArray(FunctionComponents, TEXT(":")))
+			{
+				UE_LOG(LogTargetReceiptBuildWorker, Warning, TEXT("Target receipt contains inline function version declaration that could not be parsed: '%s'"), *Property.Value);
+				return;
+			}
+
+			FGuid VersionGuid;
+			if (!FGuid::Parse(FunctionComponents[1], VersionGuid))
+			{
+				UE_LOG(LogTargetReceiptBuildWorker, Warning, TEXT("Target receipt contains inline function version guid that could not be parsed: '%s'"), *Property.Value);
+				return;
+			}
+
+			ReceiptInlineFunctionVersions.Emplace(MoveTemp(FunctionComponents[0]), MoveTemp(VersionGuid));
+		}
+		else if (Property.Name == TEXT("UnstagedRuntimeDependency"))
+		{
+			if (!TryAddExecutablePath(Property.Value))
+			{
+				UE_LOG(LogTargetReceiptBuildWorker, Warning, TEXT("Target receipt referenced an unstaged dependency that was not usable: '%s'"), *Property.Value);
+				return;
+			}
+		}
+	}
+
+	if (!bUsableBuildWorker)
+	{
+		return;
+	}
+
+	for (const FRuntimeDependency& RuntimeDependency : TargetReceipt.RuntimeDependencies)
+	{
+		if (RuntimeDependency.Path.EndsWith(TEXT(".uproject")))
+		{
+			// uproject files are not needed for build workers
+			continue;
+		}
+		if (!TryAddExecutablePath(RuntimeDependency.Path))
+		{
+			UE_LOG(LogTargetReceiptBuildWorker, Warning, TEXT("Target receipt referenced a runtime dependency that was not usable: '%s'"), *RuntimeDependency.Path);
+			return;
+		}
+	}
+	
+	if (bAbortRequested)
+	{
+		return;
+	}
+
+	// Open executable files to check for existence and gather hash info for them
+	int32 MetaIndex = 0;
+	ExecutableMeta.AddDefaulted(ExecutablePaths.Num());
+	// Consider the receipt file when computing newest executable modification time
+	FDateTime NewestExecutableModificationTime = FileManager.GetTimeStamp(*FinalTargetReceiptFilePath);
+	for (const FWorkerPath& ExecutablePath : ExecutablePaths)
+	{
+		FWorkerPathMeta& Meta = ExecutableMeta[MetaIndex];
+
+		FString LocalExecutablePath = ExecutablePath.GetLocalPath();
+
+		if (TUniquePtr<FArchive> Ar{FileManager.CreateFileReader(*LocalExecutablePath, FILEREAD_Silent)})
+		{
+			FDateTime ModificationTime = FileManager.GetTimeStamp(*LocalExecutablePath);
+			if (ModificationTime > NewestExecutableModificationTime)
+			{
+				NewestExecutableModificationTime = ModificationTime;
+			}
+			const int64 TotalSize = Ar->TotalSize();
+			FUniqueBuffer MutableBuffer = FUniqueBuffer::Alloc(uint64(TotalSize));
+			Ar->Serialize(MutableBuffer.GetData(), TotalSize);
+			Meta.Key = FIoHash::HashBuffer(MutableBuffer.GetView());
+			Meta.Value = uint64(TotalSize);
+			Ar->Close();
+		}
+		else
+		{
+			UE_LOG(LogTargetReceiptBuildWorker, Warning, TEXT("Executable file not present: '%s'"), *LocalExecutablePath);
+			return;
+		}
+		++MetaIndex;
+	}
+
+	if (ReceiptInlineBuildSystemVersion.IsValid())
+	{
+		// Accept version information declared inline in the target receipt.
+		BuildSystemVersion = MoveTemp(ReceiptInlineBuildSystemVersion);
+		FunctionVersions = MoveTemp(ReceiptInlineFunctionVersions);
+	}
+	else
+	{
+		// Gather system and function version information from the process
+		FString VersionFilePath;
+		FCbObject VersionObject;
+		if (FPaths::FileExists(FinalPublishedVersionFilePath))
+		{
+			VersionFilePath = FinalPublishedVersionFilePath;
+
+			if (TUniquePtr<FArchive> Ar{IFileManager::Get().CreateFileReader(*FinalPublishedVersionFilePath, FILEREAD_Silent)})
+			{
+				*Ar << VersionObject;
+			}
+			else
+			{
+				UE_LOG(LogTargetReceiptBuildWorker, Warning, TEXT("Could not read published version file: '%s'"), *TargetReceipt.Launch);
+				return;
+			}
+		}
+		else
+		{
+			VersionFilePath = FinalGeneratedVersionFilePath;
+
+			if (TargetReceipt.Platform != FPlatformProcess::GetBinariesSubdirectory())
+			{
+				UE_LOG(LogTargetReceiptBuildWorker, Log, TEXT("Target lacked published version information and isn't made for the current platform: '%s'"), *TargetReceipt.Launch);
+				return;
+			}
+
+			FDateTime VersionFileModificationTime = FileManager.GetTimeStamp(*FinalGeneratedVersionFilePath);
+			if (VersionFileModificationTime < NewestExecutableModificationTime)
+			{
+				FString CommandLine = TEXT("-Version=\"") + FinalGeneratedVersionFilePath + TEXT("\"");
+				FProcHandle WorkerProcHandle = FPlatformProcess::CreateProc(*FPaths::ConvertRelativePathToFull(ExecutablePaths[0].GetLocalPath()), *CommandLine, true, true, true, nullptr, 0, nullptr, nullptr);
+				if (!WorkerProcHandle.IsValid())
+				{
+					UE_LOG(LogTargetReceiptBuildWorker, Warning, TEXT("Failed to launch worker: '%s'"), *TargetReceipt.Launch);
+					return;
+				}
+
+				uint64 WorkerStartTime = FPlatformTime::Cycles64();
+				while (FPlatformProcess::IsProcRunning(WorkerProcHandle))
+				{
+					FPlatformProcess::Sleep(0.1f);
+					if (bAbortRequested || (FPlatformTime::ToSeconds64(FPlatformTime::Cycles64()-WorkerStartTime) > 10))
+					{
+						UE_LOG(LogTargetReceiptBuildWorker, Warning, TEXT("Aborted launch of worker: '%s'"), *TargetReceipt.Launch);
+						FPlatformProcess::TerminateProc(WorkerProcHandle, true);
+						FPlatformProcess::CloseProc(WorkerProcHandle);
+						return;
+					}
+				}
+
+				FPlatformProcess::CloseProc(WorkerProcHandle);
+
+				if (!FPaths::FileExists(FinalGeneratedVersionFilePath) || (FileManager.GetTimeStamp(*FinalGeneratedVersionFilePath) < NewestExecutableModificationTime))
+				{
+					UE_LOG(LogTargetReceiptBuildWorker, Warning, TEXT("Worker did not produce version file: '%s'"), *TargetReceipt.Launch);
+					return;
+				}
+			}
+
+			if (TUniquePtr<FArchive> Ar{IFileManager::Get().CreateFileReader(*FinalGeneratedVersionFilePath, FILEREAD_Silent)})
+			{
+				*Ar << VersionObject;
+			}
+			else
+			{
+				UE_LOG(LogTargetReceiptBuildWorker, Warning, TEXT("Worker did not produce readable version file: '%s'"), *TargetReceipt.Launch);
+				return;
+			}
+		}
+
+		BuildSystemVersion = VersionObject["BuildSystemVersion"].AsUuid();
+		if (!BuildSystemVersion.IsValid())
+		{
+			UE_LOG(LogTargetReceiptBuildWorker, Warning, TEXT("Invalid build system version in version file: '%s'"), *VersionFilePath);
+			return;
+		}
+
+		FCbArrayView FunctionsArrayView = VersionObject["Functions"].AsArrayView();
+		for (FCbFieldView FunctionFieldView : FunctionsArrayView)
+		{
+			FCbObjectView FunctionObjectView = FunctionFieldView.AsObjectView();
+			FunctionVersions.Emplace(FUTF8ToTCHAR(FunctionObjectView["Name"].AsString()), FunctionObjectView["Version"].AsUuid());
+		}
+	}
+
+	IModularFeatures::Get().RegisterModularFeature(UE::DerivedData::IBuildWorkerFactory::GetFeatureName(), GetWorkerFactory());
+	bEnabled = true;
+}
diff --git a/Engine/Source/Developer/DesktopPlatform/Public/MultiPlatformTargetReceiptBuildWorkers.h b/Engine/Source/Developer/DesktopPlatform/Public/MultiPlatformTargetReceiptBuildWorkers.h
new file mode 100644
--- /dev/null
+++ b/Engine/Source/Developer/DesktopPlatform/Public/MultiPlatformTargetReceiptBuildWorkers.h
@@ -0,0 +1,30 @@
+// Copyright Epic Games, Inc. All Rights Reserved.
+
+#pragma once
+
+#include "TargetReceiptBuildWorker.h"
+#include "Templates/TypeCompatibleBytes.h"
+
+/**
+ * Globally registers a UE::DerivedData::IBuildWorkerFactory instance for each platform that build workers can be supported.
+ * Users should include a $(Platform) variable in their file path to indicate that this worker receipt can exist for any platform.
+ * If the receipt path does not include a $(Platform) variable, then no attempt will be made to find variations of it for other platforms.
+*/
+class DESKTOPPLATFORM_API FMultiPlatformTargetReceiptBuildWorkers
+{
+public:
+	FMultiPlatformTargetReceiptBuildWorkers(const TCHAR* TargetReceiptFilePath);
+	~FMultiPlatformTargetReceiptBuildWorkers();
+
+private:
+	enum ESupportedPlatform : uint8
+	{
+		SupportedPlatform_Win64 = 0,
+		SupportedPlatform_Mac,
+		SupportedPlatform_Linux,
+
+		SupportedPlatform_MAX
+	};
+	TTypeCompatibleBytes<FTargetReceiptBuildWorker> PlatformSpecificWorkerFactories[SupportedPlatform_MAX];
+	bool bAllPlatformsInitialized;
+};
diff --git a/Engine/Source/Developer/DesktopPlatform/Public/TargetReceipt.h b/Engine/Source/Developer/DesktopPlatform/Public/TargetReceipt.h
--- a/Engine/Source/Developer/DesktopPlatform/Public/TargetReceipt.h
+++ b/Engine/Source/Developer/DesktopPlatform/Public/TargetReceipt.h
@@ -89,9 +89,10 @@ struct DESKTOPPLATFORM_API FTargetReceipt
 	 * Read a target receipt from disk
 	 *
 	 * @param FileName The file to read from
+	 * @param bExpandVariables Whether variables should be expanded in paths or not
 	 * @return True if the file was read successfully
 	 */
-	bool Read(const FString& FileName);
+	bool Read(const FString& FileName, bool bExpandVariables = true);
 
 	/**
 	 * Gets the default path for a target receipt
diff --git a/Engine/Source/Developer/DesktopPlatform/Public/TargetReceiptBuildWorker.h b/Engine/Source/Developer/DesktopPlatform/Public/TargetReceiptBuildWorker.h
new file mode 100644
--- /dev/null
+++ b/Engine/Source/Developer/DesktopPlatform/Public/TargetReceiptBuildWorker.h
@@ -0,0 +1,75 @@
+// Copyright Epic Games, Inc. All Rights Reserved.
+
+#pragma once
+
+#include "Async/TaskGraphInterfaces.h"
+#include "Containers/Array.h"
+#include "Containers/Map.h"
+#include "Containers/StringView.h"
+#include "Containers/UnrealString.h"
+#include "IO/IoHash.h"
+#include "Misc/Guid.h"
+#include "Templates/Tuple.h"
+#include "Templates/TypeCompatibleBytes.h"
+
+struct FTargetReceipt;
+class FTargetReceiptBuildWorkerFactory;
+
+/**
+ * Globally registers a UE::DerivedData::IBuildWorkerFactory instance that runs an executable built by UnrealBuildTool.
+ * UnrealBuildTool provides the executable information through a TargetReceipt file.
+*/
+class DESKTOPPLATFORM_API FTargetReceiptBuildWorker
+{
+public:
+	FTargetReceiptBuildWorker(const TCHAR* TargetReceiptFilePath);
+	virtual ~FTargetReceiptBuildWorker();
+
+private:
+	/**
+	 * Encapsulates a path used by the worker that has both a local and remote representation.
+	*/
+	class FWorkerPath
+	{
+	public:
+		FWorkerPath(FStringView InLocalPathRoot, FStringView InLocalPathSuffix, FStringView InRemotePathRoot = FString(), FStringView InRemotePathSuffix = FString());
+
+		FString GetLocalPath() const { return LocalPath; }
+		FString GetRemotePath() const { return RemotePath; }
+	private:
+		FString LocalPath;
+		FString RemotePath;
+	};
+
+	friend class FPopulateWorkerFromTargetReceiptTask;
+	friend class FTargetReceiptBuildWorkerFactory;
+
+	FTargetReceiptBuildWorkerFactory* GetWorkerFactory();
+	bool TryAddExecutablePath(FStringView Path);
+	void PopulateWorkerFromReceipt(const FString& TargetReceiptFilePath);
+
+	/**
+	 * Dummy type used to indicate the size of an empty class (no members) with virtual functions (contains a virtual function table).
+	 * The FTargetReceiptBuildWorkerFactory type will be created and stored in the space of one of these.
+	 * Size matching is enforced with a static_assert.  The purpose of this is to avoid exposing DerivedDataCache types to
+	 * modules that specify a target receipt worker.
+	*/
+	class IPureVirtual
+	{
+		virtual void Dummy();
+	};
+
+	static const FGuid WorkerReceiptVersion;
+	TTypeCompatibleBytes<IPureVirtual> InternalFactory;
+	FGraphEventRef PopulateTaskRef;
+	FString Name;
+	FString Platform;
+	TMap<FString, FString> EnvironmentVariables;
+	TArray<FWorkerPath> ExecutablePaths;
+	using FWorkerPathMeta = TTuple<FIoHash, uint64>;
+	TArray<FWorkerPathMeta> ExecutableMeta;
+	FGuid BuildSystemVersion;
+	TArray<TPair<FString, FGuid>> FunctionVersions;
+	std::atomic<bool> bAbortRequested;
+	bool bEnabled;
+};
diff --git a/Engine/Source/Editor/UnrealEd/Private/BaseTextureBuildWorkerFactory.cpp b/Engine/Source/Editor/UnrealEd/Private/BaseTextureBuildWorkerFactory.cpp
deleted file mode 100644
diff --git a/Engine/Source/Editor/UnrealEd/Private/DerivedDataBuildWorkers.cpp b/Engine/Source/Editor/UnrealEd/Private/DerivedDataBuildWorkers.cpp
new file mode 100644
--- /dev/null
+++ b/Engine/Source/Editor/UnrealEd/Private/DerivedDataBuildWorkers.cpp
@@ -0,0 +1,20 @@
+// Copyright Epic Games, Inc. All Rights Reserved.
+
+#include "DerivedDataBuildWorkers.h"
+#include "Misc/Optional.h"
+#include "MultiPlatformTargetReceiptBuildWorkers.h"
+
+struct FDerivedDataBuildWorkers
+{
+	FMultiPlatformTargetReceiptBuildWorkers BaseTextureBuildWorkerFactory {TEXT("$(EngineDir)/Binaries/$(Platform)/BaseTextureBuildWorker.Target")};
+};
+
+TOptional<FDerivedDataBuildWorkers> GBuildWorkers;
+
+void InitDerivedDataBuildWorkers()
+{
+	if (!GBuildWorkers.IsSet())
+	{
+		GBuildWorkers.Emplace();
+	}
+}
diff --git a/Engine/Source/Editor/UnrealEd/Private/DerivedDataBuildWorkers.h b/Engine/Source/Editor/UnrealEd/Private/DerivedDataBuildWorkers.h
new file mode 100644
--- /dev/null
+++ b/Engine/Source/Editor/UnrealEd/Private/DerivedDataBuildWorkers.h
@@ -0,0 +1,5 @@
+// Copyright Epic Games, Inc. All Rights Reserved.
+
+#pragma once
+
+void InitDerivedDataBuildWorkers();
diff --git a/Engine/Source/Editor/UnrealEd/Private/EditorEngine.cpp b/Engine/Source/Editor/UnrealEd/Private/EditorEngine.cpp
--- a/Engine/Source/Editor/UnrealEd/Private/EditorEngine.cpp
+++ b/Engine/Source/Editor/UnrealEd/Private/EditorEngine.cpp
@@ -237,6 +237,7 @@
 #include "Misc/PackageAccessTracking.h"
 #include "Misc/PackageAccessTrackingOps.h"
 #include "DerivedDataBuildRemoteExecutor.h"
+#include "DerivedDataBuildWorkers.h"
 
 #if WITH_CHAOS
 #include "ChaosSolversModule.h"
@@ -632,6 +633,7 @@ void UEditorEngine::InitEditor(IEngineLoop* InEngineLoop)
 	// Allow remote execution of derived data builds from this point
 	// TODO: This needs to be enabled earlier to allow early data builds to be remote executed.
 	InitDerivedDataBuildRemoteExecutor();
+	InitDerivedDataBuildWorkers();
 
 	// Call base.
 	UEngine::Init(InEngineLoop);
diff --git a/Engine/Source/Programs/DerivedDataBuildWorker/DerivedDataBuildWorker.Build.cs b/Engine/Source/Programs/DerivedDataBuildWorker/DerivedDataBuildWorker.Build.cs
--- a/Engine/Source/Programs/DerivedDataBuildWorker/DerivedDataBuildWorker.Build.cs
+++ b/Engine/Source/Programs/DerivedDataBuildWorker/DerivedDataBuildWorker.Build.cs
@@ -75,6 +75,7 @@ public class DerivedDataBuildWorker : ModuleRules
 			PrivateDependencyModuleNames.Add("ApplicationCore");
 		}
 
-		AdditionalPropertiesForReceipt.Add("DerivedDataBuildWorker", "true");
+		// This needs to match the version defined in TargetReceiptBuildWorkerFactory.cpp
+		AdditionalPropertiesForReceipt.Add("DerivedDataBuildWorkerReceiptVersion", "a547b2fa-e2f9-47c6-80c8-df688fb99f34");
 	}
 }
