Fixing mapping of project directory during remote execution and tweaking timeout limit.

diff --git a/Engine/Source/Developer/DesktopPlatform/Private/TargetReceiptBuildWorker.cpp b/Engine/Source/Developer/DesktopPlatform/Private/TargetReceiptBuildWorker.cpp
--- a/Engine/Source/Developer/DesktopPlatform/Private/TargetReceiptBuildWorker.cpp
+++ b/Engine/Source/Developer/DesktopPlatform/Private/TargetReceiptBuildWorker.cpp
@@ -196,7 +196,16 @@ bool FTargetReceiptBuildWorker::TryAddExecutablePath(FStringView Path)
 	}
 	else if (PrefixVariable == TEXT("ProjectDir"))
 	{
-		ExecutablePaths.Emplace(FPaths::ProjectDir(), Path.RightChop(EndVariableIndex + 2), TEXT("Project"));
+		FString ProjectDir = FPaths::ConvertRelativePathToFull(FPaths::ProjectDir());
+		FString RootDir = FPaths::ConvertRelativePathToFull(FPaths::RootDir());
+		if (FPaths::MakePathRelativeTo(ProjectDir, *RootDir))
+		{
+			ExecutablePaths.Emplace(FPaths::ProjectDir(), Path.RightChop(EndVariableIndex + 2), ProjectDir);
+		}
+		else
+		{
+			ExecutablePaths.Emplace(FPaths::ProjectDir(), Path.RightChop(EndVariableIndex + 2), TEXT("Project"));
+		}
 	}
 	else
 	{
