Fix build configuration visibility for installed builds in the project packaging editor settings
Work around circular reference by moving InstalledPlatformInfo to the DesktopPlatform module

diff --git a/Engine/Source/Developer/Android/AndroidPlatformEditor/AndroidPlatformEditor.Build.cs b/Engine/Source/Developer/Android/AndroidPlatformEditor/AndroidPlatformEditor.Build.cs
--- a/Engine/Source/Developer/Android/AndroidPlatformEditor/AndroidPlatformEditor.Build.cs
+++ b/Engine/Source/Developer/Android/AndroidPlatformEditor/AndroidPlatformEditor.Build.cs
@@ -24,7 +24,7 @@ public class AndroidPlatformEditor : ModuleRules
 				"SourceControl",
 				"AndroidRuntimeSettings",
                 "AndroidDeviceDetection",
-                "TargetPlatform",
+                "DesktopPlatform",
                 "RenderCore",
                 "RHI",
                 "MaterialShaderQualitySettings",
diff --git a/Engine/Source/Developer/DesktopPlatform/DesktopPlatform.Build.cs b/Engine/Source/Developer/DesktopPlatform/DesktopPlatform.Build.cs
--- a/Engine/Source/Developer/DesktopPlatform/DesktopPlatform.Build.cs
+++ b/Engine/Source/Developer/DesktopPlatform/DesktopPlatform.Build.cs
@@ -12,6 +12,7 @@ public class DesktopPlatform : ModuleRules
 			new string[] {
 				"DerivedDataCache",
 				"SlateFontDialog",
+				"LauncherPlatform",
 			}
 		);
 
diff --git a/Engine/Source/Developer/DesktopPlatform/Private/InstalledPlatformInfo.cpp b/Engine/Source/Developer/DesktopPlatform/Private/InstalledPlatformInfo.cpp
new file mode 100644
--- /dev/null
+++ b/Engine/Source/Developer/DesktopPlatform/Private/InstalledPlatformInfo.cpp
@@ -0,0 +1,296 @@
+// Copyright Epic Games, Inc. All Rights Reserved.
+
+#include "InstalledPlatformInfo.h"
+#include "Misc/MessageDialog.h"
+#include "Misc/Paths.h"
+#include "Misc/ConfigCacheIni.h"
+#include "Misc/App.h"
+#include "PlatformInfo.h"
+#include "IDesktopPlatform.h"
+#include "DesktopPlatformModule.h"
+#include "ILauncherPlatform.h"
+#include "LauncherPlatformModule.h"
+
+#define LOCTEXT_NAMESPACE "InstalledPlatformInfo"
+
+DEFINE_LOG_CATEGORY_STATIC(LogInstalledPlatforms, Log, All);
+
+EProjectType EProjectTypeFromString(const FString& ProjectTypeName)
+{
+	if (FCString::Strcmp(*ProjectTypeName, TEXT("Any")) == 0)
+	{
+		return EProjectType::Any;
+	}
+	else if (FCString::Strcmp(*ProjectTypeName, TEXT("Code")) == 0)
+	{
+		return EProjectType::Code;
+	}
+	else if (FCString::Strcmp(*ProjectTypeName, TEXT("Content")) == 0)
+	{
+		return EProjectType::Content;
+	}
+	else
+	{
+		return EProjectType::Unknown;
+	}
+}
+
+FInstalledPlatformInfo::FInstalledPlatformInfo()
+{
+	TArray<FString> InstalledPlatforms;
+	GConfig->GetArray(TEXT("InstalledPlatforms"), TEXT("InstalledPlatformConfigurations"), InstalledPlatforms, GEngineIni);
+
+	for (FString& InstalledPlatform : InstalledPlatforms)
+	{
+		ParsePlatformConfiguration(InstalledPlatform);
+	}
+}
+
+void FInstalledPlatformInfo::ParsePlatformConfiguration(FString PlatformConfiguration)
+{
+	// Trim whitespace at the beginning.
+	PlatformConfiguration.TrimStartInline();
+	// Remove brackets.
+	PlatformConfiguration.RemoveFromStart(TEXT("("));
+	PlatformConfiguration.RemoveFromEnd(TEXT(")"));
+
+	bool bCanCreateEntry = true;
+
+	FString ConfigurationName;
+	EBuildConfiguration Configuration = EBuildConfiguration::Unknown;
+	if (!FParse::Value(*PlatformConfiguration, TEXT("Configuration="), ConfigurationName) || !LexTryParseString(Configuration, *ConfigurationName) || Configuration == EBuildConfiguration::Unknown)
+	{
+		UE_LOG(LogInstalledPlatforms, Warning, TEXT("Unable to read configuration from %s"), *PlatformConfiguration);
+		bCanCreateEntry = false;
+	}
+
+	FString	PlatformName;
+	if (!FParse::Value(*PlatformConfiguration, TEXT("PlatformName="), PlatformName))
+	{
+		UE_LOG(LogInstalledPlatforms, Warning, TEXT("Unable to read platform from %s"), *PlatformConfiguration);
+		bCanCreateEntry = false;
+	}
+
+	FString PlatformTypeName;
+	EBuildTargetType PlatformType;
+	if (!FParse::Value(*PlatformConfiguration, TEXT("PlatformType="), PlatformTypeName) || !LexTryParseString(PlatformType, *PlatformTypeName))
+	{
+		PlatformType = EBuildTargetType::Game;
+	}
+
+	FString Architecture;
+	FParse::Value(*PlatformConfiguration, TEXT("Architecture="), Architecture);
+
+	FString RequiredFile;
+	if (FParse::Value(*PlatformConfiguration, TEXT("RequiredFile="), RequiredFile))
+	{
+		RequiredFile = FPaths::Combine(*FPaths::RootDir(), *RequiredFile);
+	}
+
+	FString ProjectTypeName;
+	EProjectType ProjectType =  EProjectType::Any;
+	if (FParse::Value(*PlatformConfiguration, TEXT("ProjectType="), ProjectTypeName))
+	{
+		ProjectType = EProjectTypeFromString(ProjectTypeName);
+	}
+	if (ProjectType == EProjectType::Unknown)
+	{
+		UE_LOG(LogInstalledPlatforms, Warning, TEXT("Unable to read project type from %s"), *PlatformConfiguration);
+		bCanCreateEntry = false;
+	}
+
+	bool bCanBeDisplayed = false;
+	FParse::Bool(*PlatformConfiguration, TEXT("bCanBeDisplayed="), bCanBeDisplayed);
+	
+	if (bCanCreateEntry)
+	{
+		FInstalledPlatformConfiguration NewConfig = {Configuration, PlatformName, PlatformType, Architecture, RequiredFile, ProjectType, bCanBeDisplayed};
+		InstalledPlatformConfigurations.Add(NewConfig);
+	}
+}
+
+bool FInstalledPlatformInfo::IsValidConfiguration(const EBuildConfiguration Configuration, EProjectType ProjectType) const
+{
+	return ContainsValidConfiguration(
+		[Configuration, ProjectType](const FInstalledPlatformConfiguration& CurConfig)
+		{
+			return CurConfig.Configuration == Configuration
+				&& (ProjectType == EProjectType::Any || CurConfig.ProjectType == EProjectType::Any
+					|| CurConfig.ProjectType == ProjectType);
+		}
+	);
+}
+
+bool FInstalledPlatformInfo::IsValidPlatform(const FString& PlatformName, EProjectType ProjectType) const
+{
+	return ContainsValidConfiguration(
+		[PlatformName, ProjectType](const FInstalledPlatformConfiguration& CurConfig)
+		{
+			return CurConfig.PlatformName == PlatformName
+				&& (ProjectType == EProjectType::Any || CurConfig.ProjectType == EProjectType::Any
+					|| CurConfig.ProjectType == ProjectType);
+		}
+	);
+}
+
+bool FInstalledPlatformInfo::IsValidPlatformAndConfiguration(const EBuildConfiguration Configuration, const FString& PlatformName, EProjectType ProjectType) const
+{
+	return ContainsValidConfiguration(
+		[Configuration, PlatformName, ProjectType](const FInstalledPlatformConfiguration& CurConfig)
+		{
+			return CurConfig.Configuration == Configuration
+				&& CurConfig.PlatformName == PlatformName
+				&& (ProjectType == EProjectType::Any || CurConfig.ProjectType == EProjectType::Any
+					|| CurConfig.ProjectType == ProjectType);
+		}
+	);
+}
+
+bool FInstalledPlatformInfo::CanDisplayPlatform(const FString& PlatformName, EProjectType ProjectType) const
+{
+	return ContainsMatchingConfiguration(
+		[PlatformName, ProjectType](const FInstalledPlatformConfiguration& CurConfig)
+		{
+			return CurConfig.PlatformName == PlatformName && (CurConfig.bCanBeDisplayed
+				|| ProjectType == EProjectType::Any || CurConfig.ProjectType == EProjectType::Any
+				|| CurConfig.ProjectType == ProjectType);
+		}
+	);
+}
+
+bool FInstalledPlatformInfo::IsValidTargetType(EBuildTargetType TargetType) const
+{
+	return ContainsValidConfiguration(
+		[TargetType](const FInstalledPlatformConfiguration& CurConfig)
+		{
+			return CurConfig.PlatformType == TargetType;
+		}
+	);
+}
+
+bool FInstalledPlatformInfo::IsValid(TOptional<EBuildTargetType> TargetType, TOptional<FString> Platform, TOptional<EBuildConfiguration> Configuration, EProjectType ProjectType, EInstalledPlatformState State) const
+{
+	if (!FApp::IsEngineInstalled())
+	{
+		return true;
+	}
+
+	for (const FInstalledPlatformConfiguration& Config : InstalledPlatformConfigurations)
+	{
+		// Check whether this configuration matches all the criteria
+		if(TargetType.IsSet() && Config.PlatformType != TargetType.GetValue())
+		{
+			continue;
+		}
+		if(Platform.IsSet() && Config.PlatformName != Platform.GetValue())
+		{
+			continue;
+		}
+		if(Configuration.IsSet() && Config.Configuration != Configuration.GetValue())
+		{
+			continue;
+		}
+		if(ProjectType != EProjectType::Any && Config.ProjectType != EProjectType::Any && Config.ProjectType != ProjectType)
+		{
+			continue;
+		}
+		if(State == EInstalledPlatformState::Downloaded && !Config.RequiredFile.IsEmpty() && !FPaths::FileExists(Config.RequiredFile))
+		{
+			continue;
+		}
+
+		// Success!
+		return true;
+	}
+
+	return false;
+}
+
+bool FInstalledPlatformInfo::IsValidPlatformArchitecture(const FString& PlatformName, const FString& Architecture) const
+{
+	return ContainsValidConfiguration(
+		[PlatformName, Architecture](const FInstalledPlatformConfiguration& CurConfig)
+		{
+			return CurConfig.PlatformName == PlatformName && CurConfig.Architecture.Contains(Architecture);
+		}
+	);
+}
+
+bool FInstalledPlatformInfo::IsPlatformMissingRequiredFile(const FString& PlatformName) const
+{
+	if (FApp::IsEngineInstalled())
+	{
+		return ContainsMatchingConfiguration(
+			[PlatformName](const FInstalledPlatformConfiguration& CurConfig)
+			{
+				return CurConfig.PlatformName == PlatformName
+					&& !CurConfig.RequiredFile.IsEmpty()
+					&& !FPaths::FileExists(CurConfig.RequiredFile);
+			}
+		);
+	}
+	return false;
+}
+
+bool FInstalledPlatformInfo::OpenInstallerOptions()
+{
+	IDesktopPlatform* DesktopPlatform = FDesktopPlatformModule::Get();
+	ILauncherPlatform* LauncherPlatform = FLauncherPlatformModule::Get();
+
+	if (DesktopPlatform != nullptr && LauncherPlatform != nullptr)
+	{
+		FString CurrentIdentifier = DesktopPlatform->GetCurrentEngineIdentifier();
+		if (DesktopPlatform->IsStockEngineRelease(CurrentIdentifier))
+		{
+			if (FMessageDialog::Open(EAppMsgType::YesNo, LOCTEXT("NotInstalled_SelectedPlatform", "The Binaries for this Target Platform are not currently installed, would you like to use the Launcher to download them?")) == EAppReturnType::Yes)
+			{
+				// TODO: Ensure that this URL opens the launcher correctly before this is included in a release
+				FString InstallerURL = FString::Printf(TEXT("ue/library/engines/UE_%s/installer"), *DesktopPlatform->GetEngineDescription(CurrentIdentifier));
+				FOpenLauncherOptions OpenOptions(InstallerURL);
+				if (LauncherPlatform->OpenLauncher(OpenOptions))
+				{
+					return true;
+				}
+			}
+		}
+	}
+	return false;
+}
+
+bool FInstalledPlatformInfo::ContainsValidConfiguration(TFunctionRef<bool(const FInstalledPlatformConfiguration)> ConfigFilter) const
+{
+	if (FApp::IsEngineInstalled())
+	{
+		for (const FInstalledPlatformConfiguration& PlatformConfiguration : InstalledPlatformConfigurations)
+		{
+			// Check whether filter accepts this configuration and it has required file
+			if (ConfigFilter(PlatformConfiguration)
+				&& (PlatformConfiguration.RequiredFile.IsEmpty()
+					|| FPaths::FileExists(PlatformConfiguration.RequiredFile)))
+			{
+				return true;
+			}
+		}
+		return false;
+	}
+	return true;
+}
+
+bool FInstalledPlatformInfo::ContainsMatchingConfiguration(TFunctionRef<bool(const FInstalledPlatformConfiguration)> ConfigFilter) const
+{
+	if (FApp::IsEngineInstalled())
+	{
+		for (const FInstalledPlatformConfiguration& PlatformConfiguration : InstalledPlatformConfigurations)
+		{
+			// Check whether filter accepts this configuration
+			if (ConfigFilter(PlatformConfiguration))
+			{
+				return true;
+			}
+		}
+		return false;
+	}
+	return true;
+}
+
+#undef LOCTEXT_NAMESPACE
diff --git a/Engine/Source/Developer/DesktopPlatform/Public/InstalledPlatformInfo.h b/Engine/Source/Developer/DesktopPlatform/Public/InstalledPlatformInfo.h
new file mode 100644
--- /dev/null
+++ b/Engine/Source/Developer/DesktopPlatform/Public/InstalledPlatformInfo.h
@@ -0,0 +1,147 @@
+// Copyright Epic Games, Inc. All Rights Reserved.
+
+#pragma once
+
+#include "CoreMinimal.h"
+
+// Forward declaration
+enum class EProjectType : uint8
+{
+	Unknown,
+	Any,
+	Code,
+	Content,
+};
+
+EProjectType EProjectTypeFromString(const FString& ProjectTypeName);
+
+enum class EInstalledPlatformState
+{
+	/** Query whether the platform is supported */
+	Supported,
+
+	/** Query whether the platform has been downloaded */
+	Downloaded,
+};
+
+/**
+ * Information about a single installed platform configuration
+ */
+struct FInstalledPlatformConfiguration
+{
+	/** Build Configuration of this combination */
+	EBuildConfiguration Configuration;
+
+	/** Name of the Platform for this combination */
+	FString PlatformName;
+
+	/** Type of Platform for this combination */
+	EBuildTargetType PlatformType;
+
+	/** Name of the Architecture for this combination */
+	FString Architecture;
+
+	/** Location of a file that must exist for this combination to be valid (optional) */
+	FString RequiredFile;
+
+	/** Type of project this configuration can be used for */
+	EProjectType ProjectType;
+
+	/** Whether to display this platform as an option even if it is not valid */
+	bool bCanBeDisplayed;
+};
+
+/**
+ * Singleton class for accessing information about installed platform configurations
+ */
+class DESKTOPPLATFORM_API FInstalledPlatformInfo
+{
+public:
+	/**
+	 * Accessor for singleton
+	 */
+	static FInstalledPlatformInfo& Get()
+	{
+		static FInstalledPlatformInfo InfoSingleton;
+		return InfoSingleton;
+	}
+
+	/**
+	 * Queries whether a configuration is valid for any available platform
+	 */
+	bool IsValidConfiguration(const EBuildConfiguration Configuration, EProjectType ProjectType = EProjectType::Any) const;
+
+	/**
+	 * Queries whether a platform has any valid configurations
+	 */
+	bool IsValidPlatform(const FString& PlatformName, EProjectType ProjectType = EProjectType::Any) const;
+
+	/**
+	 * Queries whether a platform and configuration combination is valid
+	 */
+	bool IsValidPlatformAndConfiguration(const EBuildConfiguration Configuration, const FString& PlatformName, EProjectType ProjectType = EProjectType::Any) const;
+
+	/**
+	 * Queries whether a platform can be displayed as an option, even if it's not supported for the specified project type
+	 */
+	bool CanDisplayPlatform(const FString& PlatformName, EProjectType ProjectType) const;
+
+	/**
+	 * Queries whether a target type is valid for any configuration
+	 */
+	bool IsValidTargetType(EBuildTargetType TargetType) const;
+
+	/** Determines whether the given target type is supported
+	 * @param TargetType The target type being built
+	 * @param Platform The platform being built
+	 * @param Configuration The configuration being built
+	 * @param ProjectType The project type required
+	 * @param State State of the given platform support
+	 * @return True if the target can be built
+	 */
+	bool IsValid(TOptional<EBuildTargetType> TargetType, TOptional<FString> Platform, TOptional<EBuildConfiguration> Configuration, EProjectType ProjectType, EInstalledPlatformState State) const;
+
+	/**
+	 * Queries whether a platform architecture is valid for any configuration
+	 * @param PlatformName Name of the platform's binary folder (eg. Win64, Android)
+	 * @param Architecture Either a full architecture name or a partial substring for CPU/GPU combinations (eg. "-armv7", "-es2")
+	 */
+	bool IsValidPlatformArchitecture(const FString& PlatformName, const FString& Architecture) const;
+
+	/**
+	 * Queries whether a platform has any missing required files
+	 */
+	bool IsPlatformMissingRequiredFile(const FString& PlatformName) const;
+
+	/**
+	 * Attempts to open the Launcher to the Installer options so that additional platforms can be downloaded
+	 *
+	 * @return false if the engine is not a stock release, user cancels action or launcher fails to load
+	 */
+	static bool OpenInstallerOptions();
+
+private:
+	/**
+	 * Constructor
+	 */
+	FInstalledPlatformInfo();
+
+	/**
+	 * Parse platform configuration info from a config file entry
+	 */
+	void ParsePlatformConfiguration(FString PlatformConfiguration);
+
+	/**
+	 * Given a filter function, checks whether any configuration passes that filter and has required file
+	 */
+	bool ContainsValidConfiguration(TFunctionRef<bool(const FInstalledPlatformConfiguration)> ConfigFilter) const;
+
+	/**
+	 * Given a filter function, checks whether any configuration passes that filter
+	 * Doesn't check whether required file exists, so that we can find platforms that can be optionally installed
+	 */
+	bool ContainsMatchingConfiguration(TFunctionRef<bool(const FInstalledPlatformConfiguration)> ConfigFilter) const;
+
+	/** List of installed platform configuration combinations */
+	TArray<FInstalledPlatformConfiguration> InstalledPlatformConfigurations;
+};
diff --git a/Engine/Source/Developer/DeveloperToolSettings/Private/ProjectPackagingSettings.cpp b/Engine/Source/Developer/DeveloperToolSettings/Private/ProjectPackagingSettings.cpp
--- a/Engine/Source/Developer/DeveloperToolSettings/Private/ProjectPackagingSettings.cpp
+++ b/Engine/Source/Developer/DeveloperToolSettings/Private/ProjectPackagingSettings.cpp
@@ -6,6 +6,7 @@
 #include "Interfaces/IProjectManager.h"
 #include "DesktopPlatformModule.h"
 #include "DeveloperToolSettingsDelegates.h"
+#include "InstalledPlatformInfo.h"
 
 #define LOCTEXT_NAMESPACE "SettingsClasses"
 
@@ -167,6 +168,7 @@ TArray<EProjectPackagingBuildConfigurations> UProjectPackagingSettings::GetValid
 	// Check if the project has code
 	FProjectStatus ProjectStatus;
 	bool bHasCode = IProjectManager::Get().QueryStatusForCurrentProject(ProjectStatus) && ProjectStatus.bCodeBasedProject;
+	EProjectType ProjectType = bHasCode ? EProjectType::Code : EProjectType::Content;
 
 	// If if does, find all the targets
 	const TArray<FTargetInfo>* Targets = nullptr;
@@ -183,16 +185,15 @@ TArray<EProjectPackagingBuildConfigurations> UProjectPackagingSettings::GetValid
 
 		// Check the target type is valid
 		const UProjectPackagingSettings::FConfigurationInfo& Info = UProjectPackagingSettings::ConfigurationInfo[Idx];
-		if (!bHasCode && Info.Configuration == EBuildConfiguration::DebugGame)
+		if (FInstalledPlatformInfo::Get().IsValid(TOptional<EBuildTargetType>(), TOptional<FString>(), Info.Configuration, ProjectType, EInstalledPlatformState::Downloaded))
 		{
-			continue;
+			Configurations.Add(PackagingConfiguration);
 		}
-
-		Configurations.Add(PackagingConfiguration);
 	}
 	return Configurations;
 }
 
+
 static const FTargetInfo* FindBestTargetInfo(const FString& TargetName, bool bContentOnlyUsesEngineTargets, bool* bOutIsProjectTarget=nullptr)
 {
 	bool bUseEngineTargets = false;
diff --git a/Engine/Source/Developer/TargetPlatform/Private/InstalledPlatformInfo.cpp b/Engine/Source/Developer/TargetPlatform/Private/InstalledPlatformInfo.cpp
deleted file mode 100644
diff --git a/Engine/Source/Developer/TargetPlatform/Public/InstalledPlatformInfo.h b/Engine/Source/Developer/TargetPlatform/Public/InstalledPlatformInfo.h
deleted file mode 100644
diff --git a/Engine/Source/Developer/TargetPlatform/TargetPlatform.Build.cs b/Engine/Source/Developer/TargetPlatform/TargetPlatform.Build.cs
--- a/Engine/Source/Developer/TargetPlatform/TargetPlatform.Build.cs
+++ b/Engine/Source/Developer/TargetPlatform/TargetPlatform.Build.cs
@@ -11,13 +11,11 @@ public class TargetPlatform : ModuleRules
 		PrivateDependencyModuleNames.Add("SlateCore");
 		PrivateDependencyModuleNames.Add("Slate");
 		PrivateDependencyModuleNames.Add("EditorStyle");
-		PrivateDependencyModuleNames.Add("DeveloperToolSettings");
 		PrivateDependencyModuleNames.Add("Projects");
 		PrivateDependencyModuleNames.Add("RenderCore");
 		PublicDependencyModuleNames.Add("DeveloperSettings");
 		PublicDependencyModuleNames.Add("AudioPlatformConfiguration");
 		PublicDependencyModuleNames.Add("DesktopPlatform");
-		PublicDependencyModuleNames.Add("LauncherPlatform");
 
 		// TextureFormat contains public headers that were historically part of TargetPlatform, so it is exposed
 		// as a public include path on TargetPlatform.
diff --git a/Engine/Source/Editor/MainFrame/MainFrame.Build.cs b/Engine/Source/Editor/MainFrame/MainFrame.Build.cs
--- a/Engine/Source/Editor/MainFrame/MainFrame.Build.cs
+++ b/Engine/Source/Editor/MainFrame/MainFrame.Build.cs
@@ -54,6 +54,7 @@ public class MainFrame : ModuleRules
 				"AssetTools",
 				"DerivedDataCache",
 				"DesktopPlatform",
+				"LauncherPlatform",
 				"GameProjectGeneration",
 				"ProjectTargetPlatformEditor",
 				"LevelEditor",
