Remove logging when no platform info is found because it can be spammy

diff --git a/Engine/Source/Developer/DesktopPlatform/Private/PlatformInfo.cpp b/Engine/Source/Developer/DesktopPlatform/Private/PlatformInfo.cpp
--- a/Engine/Source/Developer/DesktopPlatform/Private/PlatformInfo.cpp
+++ b/Engine/Source/Developer/DesktopPlatform/Private/PlatformInfo.cpp
@@ -122,8 +122,6 @@ const FTargetPlatformInfo* FindPlatformInfo(const FName& InPlatformName)
 		}
 	}
 
-	UE_LOG(LogDesktopPlatform, Warning, TEXT("Unable to find platform info for '%s'"), *InPlatformName.ToString());
-
 	return nullptr;
 }
 
