Reduce TargetReceiptBuildWorker to a no-op in non-editor builds.

diff --git a/Engine/Source/Developer/DesktopPlatform/Private/TargetReceiptBuildWorker.cpp b/Engine/Source/Developer/DesktopPlatform/Private/TargetReceiptBuildWorker.cpp
--- a/Engine/Source/Developer/DesktopPlatform/Private/TargetReceiptBuildWorker.cpp
+++ b/Engine/Source/Developer/DesktopPlatform/Private/TargetReceiptBuildWorker.cpp
@@ -142,14 +142,17 @@ FTargetReceiptBuildWorker::FTargetReceiptBuildWorker(const TCHAR* TargetReceiptF
 : bAbortRequested(false)
 , bEnabled(false)
 {
+#if WITH_EDITOR
 	static_assert(sizeof(FTargetReceiptBuildWorkerFactory) == sizeof(IPureVirtual), "FTargetReceiptBuildWorkerFactory type must always compile to something equivalent to an IPureVirtual size.");
 	FTargetReceiptBuildWorkerFactory* WorkerFactory = GetWorkerFactory();
 	new(WorkerFactory) FTargetReceiptBuildWorkerFactory();
 	PopulateTaskRef = TGraphTask<FPopulateWorkerFromTargetReceiptTask>::CreateTask().ConstructAndDispatchWhenReady(TargetReceiptFilePath, this);
+#endif
 }
 
 FTargetReceiptBuildWorker::~FTargetReceiptBuildWorker()
 {
+#if WITH_EDITOR
 	bAbortRequested = true;
 	if (PopulateTaskRef.IsValid())
 	{
@@ -162,6 +165,7 @@ FTargetReceiptBuildWorker::~FTargetReceiptBuildWorker()
 		IModularFeatures::Get().UnregisterModularFeature(UE::DerivedData::IBuildWorkerFactory::FeatureName, WorkerFactory);
 	}
 	WorkerFactory->~FTargetReceiptBuildWorkerFactory();
+#endif
 }
 
 FTargetReceiptBuildWorkerFactory* FTargetReceiptBuildWorker::GetWorkerFactory()
