- Allowing for packaging server builds from editor's Platforms menu
#pr 8947

diff --git a/Engine/Source/Developer/DesktopPlatform/Private/PlatformInfo.cpp b/Engine/Source/Developer/DesktopPlatform/Private/PlatformInfo.cpp
--- a/Engine/Source/Developer/DesktopPlatform/Private/PlatformInfo.cpp
+++ b/Engine/Source/Developer/DesktopPlatform/Private/PlatformInfo.cpp
@@ -80,12 +80,15 @@ FTargetPlatformInfo::FTargetPlatformInfo(const FString& InIniPlatformName, EBuil
 		// put the type in parens
 		DisplayString += FString::Printf(TEXT(" (%s)"), LexToString(PlatformType));
 
-		// client builds need to be plopped on the commandline, servers are handled differently? 
-		// @todo: not sure if needed to be handled specially, honestly
+		// client and server builds need to modify the commandline
 		if (PlatformType == EBuildTargetType::Client)
 		{
 			UATCommandLine += TEXT(" -client");
 		}
+		else if (PlatformType == EBuildTargetType::Server)
+		{
+			UATCommandLine += TEXT(" -server -noclient");
+		}
 
 		VanillaInfo->Flavors.AddUnique(this);
 	}
