- Moved SteamDeckDevice.h to DesktopPlatform so that Windows and Linux can share it, without making a whole module for one header

diff --git a/Engine/Source/Developer/DesktopPlatform/Public/SteamDeck/SteamDeckDevice.h b/Engine/Source/Developer/DesktopPlatform/Public/SteamDeck/SteamDeckDevice.h
new file mode 100644
--- /dev/null
+++ b/Engine/Source/Developer/DesktopPlatform/Public/SteamDeck/SteamDeckDevice.h
@@ -0,0 +1,90 @@
+// Copyright Epic Games, Inc. All Rights Reserved.
+
+#pragma once
+
+#include "CoreMinimal.h"
+#include "Common/TargetPlatformBase.h"
+#include "Misc/ConfigCacheIni.h"
+#if PLATFORM_WINDOWS
+#include "LocalPcTargetDevice.h"
+#endif
+#include "Serialization/MemoryLayout.h"
+
+// Extend SteamDevice from LocalPcDevice as we are leveraging the Windows Device for Steam Deck when building
+// for the Win64 platform. As SteamDeck consumes the exact same build it would for a normal Win64 device
+#if PLATFORM_WINDOWS
+class FSteamDeckDevice : public TLocalPcTargetDevice<true>
+{
+public:
+	FSteamDeckDevice(FString InIpAddr, FString InDeviceName, FString InUserName, const ITargetPlatform& InTargetPlatform)
+		: TLocalPcTargetDevice<true>(InTargetPlatform)
+		, IpAddr(InIpAddr)
+		, DeviceName(InDeviceName)
+		, UserName(InUserName)
+	{
+	}
+
+	virtual FString GetName() const override
+	{
+		return DeviceName;
+	}
+
+	virtual FTargetDeviceId GetId() const override
+	{
+		return FTargetDeviceId(TargetPlatform.PlatformName(), IpAddr);
+	}
+
+	virtual FString GetOperatingSystemName() override
+	{
+		return TEXT("SteamOS");
+	}
+
+	virtual bool GetUserCredentials(FString& OutUserName, FString& OutUserPassword) override
+	{
+		OutUserName = UserName;
+
+		// No need for a password, as it will use an rsa key that is part of the SteamOS Devkit Client
+		OutUserPassword = FString();
+
+		return true;
+	}
+
+	static TArray<ITargetDevicePtr> DiscoverDevices(const ITargetPlatform& GenericWindowsTP)
+	{
+		TArray<FString> EngineIniSteamDeckDevices;
+
+		// Expected ini format: +SteamDeckDevice=(IpAddr=10.1.33.19,Name=MySteamDeck,UserName=deck)
+		GConfig->GetArray(TEXT("/Script/WindowsTargetPlatform.WindowsTargetSettings"), TEXT("SteamDeckDevice"), EngineIniSteamDeckDevices, GEngineIni);
+
+		TArray<ITargetDevicePtr> SteamDevices;
+		for (const FString& Device : EngineIniSteamDeckDevices)
+		{
+			FString IpAddr;
+			FString Name;
+			FString UserName;
+			if (!FParse::Value(*Device, TEXT("IpAddr="), IpAddr))
+			{
+				continue;
+			}
+			// Name is not required, if not set use the IpAddr. This is what is displayed in the Editor
+			if (!FParse::Value(*Device, TEXT("Name="), Name))
+			{
+				Name = IpAddr;
+			}
+			if (!FParse::Value(*Device, TEXT("UserName="), UserName))
+			{
+				continue;
+			}
+
+			SteamDevices.Add(MakeShareable(new FSteamDeckDevice(IpAddr, Name, UserName, GenericWindowsTP)));
+		}
+
+		return SteamDevices;
+	}
+
+private:
+	FString IpAddr;
+	FString DeviceName;
+	FString UserName;
+};
+#endif // PLATFORM_WINDOWS
diff --git a/Engine/Source/Developer/Windows/WindowsTargetPlatform/Private/GenericWindowsTargetPlatform.h b/Engine/Source/Developer/Windows/WindowsTargetPlatform/Private/GenericWindowsTargetPlatform.h
--- a/Engine/Source/Developer/Windows/WindowsTargetPlatform/Private/GenericWindowsTargetPlatform.h
+++ b/Engine/Source/Developer/Windows/WindowsTargetPlatform/Private/GenericWindowsTargetPlatform.h
@@ -9,7 +9,7 @@
 #include "LocalPcTargetDevice.h"
 #endif
 #include "Serialization/MemoryLayout.h"
-#include "SteamDeckDevice.h"
+#include "SteamDeck/SteamDeckDevice.h"
 
 #if WITH_ENGINE
 	#include "Sound/SoundWave.h"
diff --git a/Engine/Source/Developer/Windows/WindowsTargetPlatform/Private/SteamDeckDevice.h b/Engine/Source/Developer/Windows/WindowsTargetPlatform/Private/SteamDeckDevice.h
deleted file mode 100644
