- Windows and Linux SteamDeck devices in C++

diff --git a/Engine/Source/Developer/DesktopPlatform/Public/SteamDeck/SteamDeckDevice.h b/Engine/Source/Developer/DesktopPlatform/Public/SteamDeck/SteamDeckDevice.h
--- a/Engine/Source/Developer/DesktopPlatform/Public/SteamDeck/SteamDeckDevice.h
+++ b/Engine/Source/Developer/DesktopPlatform/Public/SteamDeck/SteamDeckDevice.h
@@ -5,23 +5,20 @@
 #include "CoreMinimal.h"
 #include "Common/TargetPlatformBase.h"
 #include "Misc/ConfigCacheIni.h"
-#if PLATFORM_WINDOWS
-#include "LocalPcTargetDevice.h"
-#endif
 #include "Serialization/MemoryLayout.h"
 
-// Extend SteamDevice from LocalPcDevice as we are leveraging the Windows Device for Steam Deck when building
-// for the Win64 platform. As SteamDeck consumes the exact same build it would for a normal Win64 device
-#if PLATFORM_WINDOWS
-class FSteamDeckDevice : public TLocalPcTargetDevice<true>
+// Extend from either WIndows or Linux device, since there will be a lot of shared functionality
+template<class ParentDeviceClass>
+class TSteamDeckDevice : public ParentDeviceClass
 {
 public:
-	FSteamDeckDevice(FString InIpAddr, FString InDeviceName, FString InUserName, const ITargetPlatform& InTargetPlatform)
-		: TLocalPcTargetDevice<true>(InTargetPlatform)
+	TSteamDeckDevice(FString InIpAddr, FString InDeviceName, FString InUserName, const ITargetPlatform& InTargetPlatform, const TCHAR* InRuntimeOSName)
+		: ParentDeviceClass(InTargetPlatform)
 		, IpAddr(InIpAddr)
-		, DeviceName(InDeviceName)
 		, UserName(InUserName)
+		, RuntimeOSName(InRuntimeOSName)
 	{
+		DeviceName = FString::Printf(TEXT("%s (%s)"), *InDeviceName, InRuntimeOSName);
 	}
 
 	virtual FString GetName() const override
@@ -31,12 +28,12 @@ public:
 
 	virtual FTargetDeviceId GetId() const override
 	{
-		return FTargetDeviceId(TargetPlatform.PlatformName(), IpAddr);
+		return FTargetDeviceId(this->TargetPlatform.PlatformName(), IpAddr);
 	}
 
 	virtual FString GetOperatingSystemName() override
 	{
-		return TEXT("SteamOS");
+		return FString::Printf(TEXT("SteamOS (%s)"), *RuntimeOSName);
 	}
 
 	virtual bool GetUserCredentials(FString& OutUserName, FString& OutUserPassword) override
@@ -49,12 +46,12 @@ public:
 		return true;
 	}
 
-	static TArray<ITargetDevicePtr> DiscoverDevices(const ITargetPlatform& GenericWindowsTP)
+	static TArray<ITargetDevicePtr> DiscoverDevices(const ITargetPlatform& TargetPlatform, const TCHAR* RuntimeOSName)
 	{
 		TArray<FString> EngineIniSteamDeckDevices;
 
 		// Expected ini format: +SteamDeckDevice=(IpAddr=10.1.33.19,Name=MySteamDeck,UserName=deck)
-		GConfig->GetArray(TEXT("/Script/WindowsTargetPlatform.WindowsTargetSettings"), TEXT("SteamDeckDevice"), EngineIniSteamDeckDevices, GEngineIni);
+		GConfig->GetArray(TEXT("SteamDeck"), TEXT("SteamDeckDevice"), EngineIniSteamDeckDevices, GEngineIni);
 
 		TArray<ITargetDevicePtr> SteamDevices;
 		for (const FString& Device : EngineIniSteamDeckDevices)
@@ -76,7 +73,7 @@ public:
 				continue;
 			}
 
-			SteamDevices.Add(MakeShareable(new FSteamDeckDevice(IpAddr, Name, UserName, GenericWindowsTP)));
+			SteamDevices.Add(MakeShareable(new TSteamDeckDevice<ParentDeviceClass>(IpAddr, Name, UserName, TargetPlatform, RuntimeOSName)));
 		}
 
 		return SteamDevices;
@@ -86,5 +83,5 @@ private:
 	FString IpAddr;
 	FString DeviceName;
 	FString UserName;
+	FString RuntimeOSName;
 };
-#endif // PLATFORM_WINDOWS
diff --git a/Engine/Source/Developer/Linux/LinuxTargetPlatform/Private/LinuxTargetDevice.h b/Engine/Source/Developer/Linux/LinuxTargetPlatform/Private/LinuxTargetDevice.h
--- a/Engine/Source/Developer/Linux/LinuxTargetPlatform/Private/LinuxTargetDevice.h
+++ b/Engine/Source/Developer/Linux/LinuxTargetPlatform/Private/LinuxTargetDevice.h
@@ -49,6 +49,11 @@ public:
 		, SavePlatformDevices(InSavePlatformDevices)
 	{ }
 
+	/**
+	 * Minimal constructor for use with SteamDeck devices */
+	FLinuxTargetDevice(const ITargetPlatform& InTargetPlatform)
+		: FLinuxTargetDevice(InTargetPlatform, TEXT("UnknownName"), nullptr) 
+	{ }
 
 public:
 
@@ -195,7 +200,7 @@ public:
 		return false;
 	}
 
-private:
+protected:
 
 	// Holds a reference to the device's target platform.
 	const ITargetPlatform& TargetPlatform;
diff --git a/Engine/Source/Developer/Linux/LinuxTargetPlatform/Private/LinuxTargetPlatform.h b/Engine/Source/Developer/Linux/LinuxTargetPlatform/Private/LinuxTargetPlatform.h
--- a/Engine/Source/Developer/Linux/LinuxTargetPlatform/Private/LinuxTargetPlatform.h
+++ b/Engine/Source/Developer/Linux/LinuxTargetPlatform/Private/LinuxTargetPlatform.h
@@ -12,6 +12,7 @@
 #include "Misc/ConfigCacheIni.h"
 #include "Interfaces/ITargetPlatform.h"
 #include "Common/TargetPlatformBase.h"
+#include "SteamDeck/SteamDeckDevice.h"
 
 
 #if WITH_ENGINE
@@ -60,6 +61,12 @@ public:
 
 		InitDevicesFromConfig();
 
+		if (!TProperties::IsArm64())
+		{
+			SteamDevices = TSteamDeckDevice<FLinuxTargetDevice>::DiscoverDevices(*this, TEXT("Native Linux"));
+		}
+
+
 		// Get the Target RHIs for this platform, we do not always want all those that are supported.
 		TArray<FName> TargetedShaderFormats;
 		TLinuxTargetPlatform::GetAllTargetedShaderFormats(TargetedShaderFormats);
@@ -125,6 +132,14 @@ public:
 		{
 			OutDevices.Add(DeviceIter.Value);
 		}
+
+		for (const ITargetDevicePtr& SteamDeck : SteamDevices)
+		{
+			if (SteamDeck.IsValid())
+			{
+				OutDevices.Add(SteamDeck);
+			}
+		}
 	}
 
 	virtual bool GenerateStreamingInstallManifest(const TMultiMap<FString, int32>& PakchunkMap, const TSet<int32>& PakchunkIndicesInUse) const override
@@ -157,6 +172,14 @@ public:
 			}
 		}
 
+		for (const ITargetDevicePtr& SteamDeck : SteamDevices)
+		{
+			if (SteamDeck.IsValid() && DeviceId == SteamDeck->GetId())
+			{
+				return SteamDeck;
+			}
+		}
+
 		return nullptr;
 	}
 
@@ -456,6 +479,7 @@ protected:
 	FLinuxTargetDevicePtr LocalDevice;
 	// Holds a map of valid devices.
 	TMap<FString, FLinuxTargetDevicePtr> Devices;
+	TArray<ITargetDevicePtr> SteamDevices;
 
 
 #if WITH_ENGINE
diff --git a/Engine/Source/Developer/Windows/WindowsTargetPlatform/Private/GenericWindowsTargetPlatform.h b/Engine/Source/Developer/Windows/WindowsTargetPlatform/Private/GenericWindowsTargetPlatform.h
--- a/Engine/Source/Developer/Windows/WindowsTargetPlatform/Private/GenericWindowsTargetPlatform.h
+++ b/Engine/Source/Developer/Windows/WindowsTargetPlatform/Private/GenericWindowsTargetPlatform.h
@@ -48,7 +48,7 @@ public:
 		LocalDevice = MakeShareable(new TTargetDevice(*this));
 
 		// Check if we have any SteamDeck devices around
-		SteamDevices = FSteamDeckDevice::DiscoverDevices(*this);
+		SteamDevices = TSteamDeckDevice<TLocalPcTargetDevice<true>>::DiscoverDevices(*this, TEXT("Proton"));
 #endif
 
 	#if WITH_ENGINE
diff --git a/Engine/Source/Programs/AutomationTool/Linux/LinuxPlatform.Automation.cs b/Engine/Source/Programs/AutomationTool/Linux/LinuxPlatform.Automation.cs
--- a/Engine/Source/Programs/AutomationTool/Linux/LinuxPlatform.Automation.cs
+++ b/Engine/Source/Programs/AutomationTool/Linux/LinuxPlatform.Automation.cs
@@ -35,7 +35,11 @@ public abstract class BaseLinuxPlatform : Platform
 			Devices.Add(LocalMachine);
 		}
 
-		Devices.AddRange(SteamDeckSupport.GetDevices(UnrealTargetPlatform.Linux));
+		// only add for true Linux, not Arm64, etc
+		if (PlatformType == UnrealTargetPlatform.Linux)
+		{
+			Devices.AddRange(SteamDeckSupport.GetDevices(UnrealTargetPlatform.Linux));
+		}
 
 		return Devices.ToArray();
 	}
diff --git a/Engine/Source/Programs/AutomationTool/SteamDeck/SteamDeckSupport.cs b/Engine/Source/Programs/AutomationTool/SteamDeck/SteamDeckSupport.cs
--- a/Engine/Source/Programs/AutomationTool/SteamDeck/SteamDeckSupport.cs
+++ b/Engine/Source/Programs/AutomationTool/SteamDeck/SteamDeckSupport.cs
@@ -26,7 +26,7 @@ public static class SteamDeckSupport
 		ConfigHierarchy EngineConfig = ConfigCache.ReadHierarchy(ConfigHierarchyType.Engine, null, UnrealTargetPlatform.Win64);
 
 		List<string> SteamDeckDevices;
-		if (EngineConfig.GetArray("/Script/WindowsTargetPlatform.WindowsTargetSettings", "SteamDeckDevice", out SteamDeckDevices))
+		if (EngineConfig.GetArray("SteamDeck", "SteamDeckDevice", out SteamDeckDevices))
 		{
 			// Expected ini format: +SteamDeckDevice=(IpAddr=10.1.33.19,Name=MySteamDeck,UserName=deck)
 			foreach (string DeckDevice in SteamDeckDevices)
