Fix the detection of UBT running.  This prevents an issue where hot reload would try and load the new module faster than it should.

diff --git a/Engine/Source/Developer/DesktopPlatform/Private/Linux/DesktopPlatformLinux.cpp b/Engine/Source/Developer/DesktopPlatform/Private/Linux/DesktopPlatformLinux.cpp
--- a/Engine/Source/Developer/DesktopPlatform/Private/Linux/DesktopPlatformLinux.cpp
+++ b/Engine/Source/Developer/DesktopPlatform/Private/Linux/DesktopPlatformLinux.cpp
@@ -450,7 +450,55 @@ bool FDesktopPlatformLinux::RunUnrealBuildTool(const FText& Description, const F
 
 bool FDesktopPlatformLinux::IsUnrealBuildToolRunning()
 {
-	return FPlatformProcess::IsApplicationRunning(TEXT("UnrealBuildTool"));
+	//TODO 5.1, make this a common routine instead of duplicated
+	FString RunsDir = FPaths::Combine(FPaths::EngineIntermediateDir(), TEXT("UbtRuns"));
+	if (!FPaths::DirectoryExists(RunsDir))
+	{
+		return false;
+	}
+
+	bool bIsRunning = false;
+	IFileManager::Get().IterateDirectory(*RunsDir, [&bIsRunning](const TCHAR* Pathname, bool bIsDirectory)
+		{
+			if (!bIsDirectory)
+			{
+				bool bDeleteFile = true;
+
+				FString Filename = FPaths::GetBaseFilename(FString(Pathname));
+				const TCHAR* Delim = FCString::Strchr(*Filename, '_');
+				if (Delim != nullptr)
+				{
+					FString Pid(*Filename, Delim - *Filename);
+					int ProcessId = 0;
+					LexFromString(ProcessId, *Pid);
+					const FString EntryFullPath = FPlatformProcess::GetApplicationName(ProcessId);
+					if (!EntryFullPath.IsEmpty())
+					{
+						FString EntryFullPathUpper = EntryFullPath.ToUpper();
+						const FTCHARToUTF8 Utf8String(*EntryFullPathUpper);
+						FMD5Hash Hash;
+						LexFromString(Hash, Delim + 1);
+
+						FMD5 Md5Gen;
+						Md5Gen.Update(reinterpret_cast<const uint8*>(Utf8String.Get()), Utf8String.Length());
+						FMD5Hash TestHash;
+						TestHash.Set(Md5Gen);
+						if (Hash == TestHash)
+						{
+							bDeleteFile = false;
+							bIsRunning = true;
+						}
+					}
+					if (bDeleteFile)
+					{
+						IFileManager::Get().Delete(Pathname);
+					}
+				}
+			}
+			return true;
+		});
+
+	return bIsRunning;
 }
 
 FFeedbackContext* FDesktopPlatformLinux::GetNativeFeedbackContext()
diff --git a/Engine/Source/Developer/DesktopPlatform/Private/Mac/DesktopPlatformMac.cpp b/Engine/Source/Developer/DesktopPlatform/Private/Mac/DesktopPlatformMac.cpp
--- a/Engine/Source/Developer/DesktopPlatform/Private/Mac/DesktopPlatformMac.cpp
+++ b/Engine/Source/Developer/DesktopPlatform/Private/Mac/DesktopPlatformMac.cpp
@@ -627,7 +627,55 @@ bool FDesktopPlatformMac::RunUnrealBuildTool(const FText& Description, const FSt
 
 bool FDesktopPlatformMac::IsUnrealBuildToolRunning()
 {
-	return FPlatformProcess::IsApplicationRunning(TEXT("UnrealBuildTool"));
+	//TODO 5.1, make this a common routine instead of duplicated
+	FString RunsDir = FPaths::Combine(FPaths::EngineIntermediateDir(), TEXT("UbtRuns"));
+	if (!FPaths::DirectoryExists(RunsDir))
+	{
+		return false;
+	}
+
+	bool bIsRunning = false;
+	IFileManager::Get().IterateDirectory(*RunsDir, [&bIsRunning](const TCHAR* Pathname, bool bIsDirectory)
+		{
+			if (!bIsDirectory)
+			{
+				bool bDeleteFile = true;
+
+				FString Filename = FPaths::GetBaseFilename(FString(Pathname));
+				const TCHAR* Delim = FCString::Strchr(*Filename, '_');
+				if (Delim != nullptr)
+				{
+					FString Pid(*Filename, Delim - *Filename);
+					int ProcessId = 0;
+					LexFromString(ProcessId, *Pid);
+					const FString EntryFullPath = FPlatformProcess::GetApplicationName(ProcessId);
+					if (!EntryFullPath.IsEmpty())
+					{
+						FString EntryFullPathUpper = EntryFullPath.ToUpper();
+						const FTCHARToUTF8 Utf8String(*EntryFullPathUpper);
+						FMD5Hash Hash;
+						LexFromString(Hash, Delim + 1);
+
+						FMD5 Md5Gen;
+						Md5Gen.Update(reinterpret_cast<const uint8*>(Utf8String.Get()), Utf8String.Length());
+						FMD5Hash TestHash;
+						TestHash.Set(Md5Gen);
+						if (Hash == TestHash)
+						{
+							bDeleteFile = false;
+							bIsRunning = true;
+						}
+					}
+					if (bDeleteFile)
+					{
+						IFileManager::Get().Delete(Pathname);
+					}
+				}
+			}
+			return true;
+		});
+
+	return bIsRunning;
 }
 
 FFeedbackContext* FDesktopPlatformMac::GetNativeFeedbackContext()
diff --git a/Engine/Source/Developer/DesktopPlatform/Private/Windows/DesktopPlatformWindows.cpp b/Engine/Source/Developer/DesktopPlatform/Private/Windows/DesktopPlatformWindows.cpp
--- a/Engine/Source/Developer/DesktopPlatform/Private/Windows/DesktopPlatformWindows.cpp
+++ b/Engine/Source/Developer/DesktopPlatform/Private/Windows/DesktopPlatformWindows.cpp
@@ -479,30 +479,55 @@ bool FDesktopPlatformWindows::RunUnrealBuildTool(const FText& Description, const
 
 bool FDesktopPlatformWindows::IsUnrealBuildToolRunning()
 {
-	FString UBTPath = GetUnrealBuildToolExecutableFilename(FPaths::RootDir());
-	FPaths::MakePlatformFilename(UBTPath);
-
-	HANDLE SnapShot = ::CreateToolhelp32Snapshot(TH32CS_SNAPPROCESS, 0);
-	if (SnapShot != INVALID_HANDLE_VALUE)
+	//TODO 5.1, make this a common routine instead of duplicated
+	FString RunsDir = FPaths::Combine(FPaths::EngineIntermediateDir(), TEXT("UbtRuns"));
+	if (!FPaths::DirectoryExists(RunsDir))
 	{
-		PROCESSENTRY32 Entry;
-		Entry.dwSize = sizeof(PROCESSENTRY32);
+		return false;
+	}
 
-		if (::Process32First(SnapShot, &Entry))
+	bool bIsRunning = false;
+	IFileManager::Get().IterateDirectory(*RunsDir, [&bIsRunning](const TCHAR* Pathname, bool bIsDirectory)
 		{
-			do
+			if (!bIsDirectory)
 			{
-				const FString EntryFullPath = FPlatformProcess::GetApplicationName(Entry.th32ProcessID);
-				if (EntryFullPath == UBTPath)
+				bool bDeleteFile = true;
+
+				FString Filename = FPaths::GetBaseFilename(FString(Pathname));
+				const TCHAR* Delim = FCString::Strchr(*Filename, '_');
+				if (Delim != nullptr)
 				{
-					::CloseHandle(SnapShot);
-					return true;
+					FString Pid(*Filename, Delim - *Filename);
+					int ProcessId = 0;
+					LexFromString(ProcessId, *Pid);
+					const FString EntryFullPath = FPlatformProcess::GetApplicationName(ProcessId);
+					if (!EntryFullPath.IsEmpty())
+					{
+						FString EntryFullPathUpper = EntryFullPath.ToUpper();
+						const FTCHARToUTF8 Utf8String(*EntryFullPathUpper);
+						FMD5Hash Hash;
+						LexFromString(Hash, Delim + 1);
+
+						FMD5 Md5Gen;
+						Md5Gen.Update(reinterpret_cast<const uint8*>(Utf8String.Get()), Utf8String.Length());
+						FMD5Hash TestHash;
+						TestHash.Set(Md5Gen);
+						if (Hash == TestHash)
+						{
+							bDeleteFile = false;
+							bIsRunning = true;
+						}
+					}
+					if (bDeleteFile)
+					{
+						IFileManager::Get().Delete(Pathname);
+					}
 				}
-			} while (::Process32Next(SnapShot, &Entry));
-		}
-	}
-	::CloseHandle(SnapShot);
-	return false;
+			}
+			return true;
+		});
+
+	return bIsRunning;
 }
 
 FFeedbackContext* FDesktopPlatformWindows::GetNativeFeedbackContext()
diff --git a/Engine/Source/Programs/UnrealBuildTool/UnrealBuildTool.cs b/Engine/Source/Programs/UnrealBuildTool/UnrealBuildTool.cs
--- a/Engine/Source/Programs/UnrealBuildTool/UnrealBuildTool.cs
+++ b/Engine/Source/Programs/UnrealBuildTool/UnrealBuildTool.cs
@@ -444,6 +444,7 @@ namespace UnrealBuildTool
 		/// <returns>Zero on success, non-zero on error</returns>
 		private static int Main(string[] ArgumentsArray)
 		{
+			FileReference? RunFile = null;
 			SingleInstanceMutex? Mutex = null;
 			JsonTracer? Tracer = null;
 			
@@ -554,6 +555,24 @@ namespace UnrealBuildTool
 					Log.AddFileWriter("LogTraceListener", Options.LogFileName);
 				}
 
+				// Create a UbtRun file
+				try
+				{
+					DirectoryReference RunsDir = DirectoryReference.Combine(Unreal.EngineDirectory, "Intermediate", "UbtRuns");
+					Directory.CreateDirectory(RunsDir.FullName);
+					string ModuleFileName = Process.GetCurrentProcess().MainModule?.FileName ?? "";
+					if (!string.IsNullOrEmpty(ModuleFileName))
+					{
+						ModuleFileName = Path.GetFullPath(ModuleFileName);
+					}
+					FileReference RunFileTemp = FileReference.Combine(RunsDir, $"{Process.GetCurrentProcess().Id}_{ContentHash.MD5(Encoding.UTF8.GetBytes(ModuleFileName.ToUpperInvariant()))}");
+					File.WriteAllLines(RunFileTemp.FullName, new string [] { ModuleFileName });
+					RunFile = RunFileTemp;
+				}
+				catch
+				{
+				}
+
 				// Acquire a lock for this branch
 				if((ModeOptions & ToolModeOptions.SingleInstance) != 0 && !Options.bNoMutex)
 				{
@@ -645,6 +664,18 @@ namespace UnrealBuildTool
 				{
 					Mutex.Dispose();
 				}
+
+				// Delete the ubt run file
+				if (RunFile != null)
+				{
+					try
+					{
+						File.Delete(RunFile.FullName);
+					}
+					catch
+					{
+					}
+				}
 			}
 		}
 	}
