Added  missing include file causing tools to fail to compile.

diff --git a/Engine/Source/Developer/DesktopPlatform/Private/Linux/DesktopPlatformLinux.cpp b/Engine/Source/Developer/DesktopPlatform/Private/Linux/DesktopPlatformLinux.cpp
--- a/Engine/Source/Developer/DesktopPlatform/Private/Linux/DesktopPlatformLinux.cpp
+++ b/Engine/Source/Developer/DesktopPlatform/Private/Linux/DesktopPlatformLinux.cpp
@@ -6,6 +6,7 @@
 #include "Misc/Guid.h"
 #include "Misc/ConfigCacheIni.h"
 #include "Misc/FeedbackContext.h"
+#include "Misc/SecureHash.h"
 #include "DesktopPlatformPrivate.h"
 #include "Modules/ModuleManager.h"
 #include "Linux/LinuxApplication.h"
diff --git a/Engine/Source/Developer/DesktopPlatform/Private/Mac/DesktopPlatformMac.cpp b/Engine/Source/Developer/DesktopPlatform/Private/Mac/DesktopPlatformMac.cpp
--- a/Engine/Source/Developer/DesktopPlatform/Private/Mac/DesktopPlatformMac.cpp
+++ b/Engine/Source/Developer/DesktopPlatform/Private/Mac/DesktopPlatformMac.cpp
@@ -8,6 +8,7 @@
 #include "Misc/Paths.h"
 #include "Misc/ConfigCacheIni.h"
 #include "Misc/Guid.h"
+#include "Misc/SecureHash.h"
 #include "HAL/FileManager.h"
 
 #define LOCTEXT_NAMESPACE "DesktopPlatform"
diff --git a/Engine/Source/Developer/DesktopPlatform/Private/Windows/DesktopPlatformWindows.cpp b/Engine/Source/Developer/DesktopPlatform/Private/Windows/DesktopPlatformWindows.cpp
--- a/Engine/Source/Developer/DesktopPlatform/Private/Windows/DesktopPlatformWindows.cpp
+++ b/Engine/Source/Developer/DesktopPlatform/Private/Windows/DesktopPlatformWindows.cpp
@@ -7,6 +7,7 @@
 #include "Microsoft/COMPointer.h"
 #include "Misc/Paths.h"
 #include "Misc/Guid.h"
+#include "Misc/SecureHash.h"
 #include "HAL/FileManager.h"
 
 #include "Windows/AllowWindowsPlatformTypes.h"
