Set a default UserName for the SteamDeck device confiuration. As of SteamOS version 3.1, "deck" is the required UserName to be used when making a remote connection to the device. Eventually it will be configurable, so we will allow users to set it via the config.

Also add a "SteamDeck" prefix to the default device name. This makes sure that its obvious what device this is until we can get a nice icon.

#rnx

diff --git a/Engine/Source/Developer/DesktopPlatform/Public/SteamDeck/SteamDeckDevice.h b/Engine/Source/Developer/DesktopPlatform/Public/SteamDeck/SteamDeckDevice.h
--- a/Engine/Source/Developer/DesktopPlatform/Public/SteamDeck/SteamDeckDevice.h
+++ b/Engine/Source/Developer/DesktopPlatform/Public/SteamDeck/SteamDeckDevice.h
@@ -6,6 +6,10 @@
 #include "Common/TargetPlatformBase.h"
 #include "Misc/ConfigCacheIni.h"
 #include "Serialization/MemoryLayout.h"
+#include "Logging/LogMacros.h"
+
+DECLARE_LOG_CATEGORY_EXTERN(LogSteamDeckDevice, Log, All);
+DEFINE_LOG_CATEGORY(LogSteamDeckDevice);
 
 // Extend from either WIndows or Linux device, since there will be a lot of shared functionality
 template<class ParentDeviceClass>
@@ -58,22 +62,32 @@ public:
 		{
 			FString IpAddr;
 			FString Name;
-			FString UserName;
+			FString ConfigUserName;
+			
+			// As of SteamOS version 3.1, "deck" is the required UserName to be used when making a remote connection
+			// to the device. Eventually it will be configurable, so we will allow users to set it via the config. 
+			static const FString DefaultUserName = TEXT("deck");
+			
 			if (!FParse::Value(*Device, TEXT("IpAddr="), IpAddr))
 			{
+				UE_LOG(LogSteamDeckDevice, Error, TEXT("You must specify the 'IpAddr' field to connect to a SteamDeck!"));
 				continue;
 			}
+			
 			// Name is not required, if not set use the IpAddr. This is what is displayed in the Editor
 			if (!FParse::Value(*Device, TEXT("Name="), Name))
 			{
-				Name = IpAddr;
+				static const FString DefaultNamePrefix = "[SteamDeck] ";
+				Name = DefaultNamePrefix + IpAddr;
 			}
-			if (!FParse::Value(*Device, TEXT("UserName="), UserName))
+
+			// If the user has not specified a UserName in the config, fallback to the default user name for the SteamDeck
+			if (!FParse::Value(*Device, TEXT("UserName="), ConfigUserName))
 			{
-				continue;
+				ConfigUserName = DefaultUserName;
 			}
 
-			SteamDevices.Add(MakeShareable(new TSteamDeckDevice<ParentDeviceClass>(IpAddr, Name, UserName, TargetPlatform, RuntimeOSName)));
+			SteamDevices.Add(MakeShareable(new TSteamDeckDevice<ParentDeviceClass>(IpAddr, Name, ConfigUserName, TargetPlatform, RuntimeOSName)));
 		}
 
 		return SteamDevices;
diff --git a/Engine/Source/Programs/AutomationTool/SteamDeck/SteamDeckSupport.cs b/Engine/Source/Programs/AutomationTool/SteamDeck/SteamDeckSupport.cs
--- a/Engine/Source/Programs/AutomationTool/SteamDeck/SteamDeckSupport.cs
+++ b/Engine/Source/Programs/AutomationTool/SteamDeck/SteamDeckSupport.cs
@@ -15,6 +15,7 @@ public static class SteamDeckSupport
 	static string DevKitRSAPath = Path.Combine(Environment.GetFolderPath(Environment.SpecialFolder.LocalApplicationData), @"steamos-devkit\steamos-devkit\devkit_rsa");
 	static string KnownHostsPath = Path.Combine(Environment.GetFolderPath(Environment.SpecialFolder.UserProfile), @".ssh\known_hosts");
 	static string SteamDeckScripts = Path.Combine(Unreal.RootDirectory.FullName, "Engine\\Build\\SteamDeck");
+	static string DefaultUserName = "deck";
 	
 	#region Devices
 
@@ -34,7 +35,11 @@ public static class SteamDeckSupport
 			{
 				string IpAddr = ConfigHierarchy.GetStructEntry(DeckDevice, "IpAddr", false);
 				string DeviceName = ConfigHierarchy.GetStructEntry(DeckDevice, "Name", false);
-				string UserName = ConfigHierarchy.GetStructEntry(DeckDevice, "UserName", false);
+				string ConfigUserName = ConfigHierarchy.GetStructEntry(DeckDevice, "UserName", false);
+
+				// As of SteamOS version 3.1, "deck" is the required UserName to be used when making a remote connection
+				// to the device. Eventually it will be configurable, so we will allow users to set it via the config.
+				string UserName = string.IsNullOrEmpty(ConfigUserName) ? DefaultUserName : ConfigUserName;
 
 				if (RuntimePlatform == UnrealTargetPlatform.Win64)
 				{
@@ -48,7 +53,7 @@ public static class SteamDeckSupport
 				// Name is optional, if its empty/not found lets just use the IpAddr for the Name
 				if (string.IsNullOrEmpty(DeviceName))
 				{
-					DeviceName = IpAddr;
+					DeviceName = "[SteamDeck] " + IpAddr;
 				}
 
 				if (!string.IsNullOrEmpty(IpAddr))
@@ -60,6 +65,10 @@ public static class SteamDeckSupport
 
 					Devices.Add(SteamDeck);
 				}
+				else
+				{
+					CommandUtils.LogError("You must specify the 'IpAddr' field to connect to a SteamDeck!");
+				}
 			}
 		}
 
