- Fixed SteamDeckDevice for monolithic editors (temp fix, the LogHAL is not the best category)

diff --git a/Engine/Source/Developer/DesktopPlatform/Public/SteamDeck/SteamDeckDevice.h b/Engine/Source/Developer/DesktopPlatform/Public/SteamDeck/SteamDeckDevice.h
--- a/Engine/Source/Developer/DesktopPlatform/Public/SteamDeck/SteamDeckDevice.h
+++ b/Engine/Source/Developer/DesktopPlatform/Public/SteamDeck/SteamDeckDevice.h
@@ -8,9 +8,6 @@
 #include "Serialization/MemoryLayout.h"
 #include "Logging/LogMacros.h"
 
-DECLARE_LOG_CATEGORY_EXTERN(LogSteamDeckDevice, Log, All);
-DEFINE_LOG_CATEGORY(LogSteamDeckDevice);
-
 // Extend from either WIndows or Linux device, since there will be a lot of shared functionality
 template<class ParentDeviceClass>
 class TSteamDeckDevice : public ParentDeviceClass
@@ -70,7 +67,7 @@ public:
 			
 			if (!FParse::Value(*Device, TEXT("IpAddr="), IpAddr))
 			{
-				UE_LOG(LogSteamDeckDevice, Error, TEXT("You must specify the 'IpAddr' field to connect to a SteamDeck!"));
+				UE_LOG(LogHAL, Error, TEXT("You must specify the 'IpAddr' field to connect to a SteamDeck!"));
 				continue;
 			}
 			
