PlatformInfo fix bad string argument in checkf printf

diff --git a/Engine/Source/Developer/DesktopPlatform/Private/PlatformInfo.cpp b/Engine/Source/Developer/DesktopPlatform/Private/PlatformInfo.cpp
--- a/Engine/Source/Developer/DesktopPlatform/Private/PlatformInfo.cpp
+++ b/Engine/Source/Developer/DesktopPlatform/Private/PlatformInfo.cpp
@@ -45,7 +45,7 @@ FTargetPlatformInfo::FTargetPlatformInfo(const FString& InIniPlatformName, EBuil
  			return Item->Name == IniPlatformName;
  		});
 
-		checkf(FoundInfo != nullptr, TEXT("Creating a TargetPlatform (%s, %s, %s) that needed a 'vanilla' TP already created, but it wasn't found. Create Game TPs first, and all cook flavors last"), *InIniPlatformName, LexToString(InType), &InCookFlavor);
+		checkf(FoundInfo != nullptr, TEXT("Creating a TargetPlatform (%s, %s, %s) that needed a 'vanilla' TP already created, but it wasn't found. Create Game TPs first, and all cook flavors last"), *InIniPlatformName, LexToString(InType), *InCookFlavor);
 
 		VanillaInfo = *FoundInfo;
 	}
