Support for using OidcToken in editor / cooker to acquire a user token, currently opt in.

diff --git a/Engine/Source/Developer/DerivedDataCache/Private/HttpCacheStore.cpp b/Engine/Source/Developer/DerivedDataCache/Private/HttpCacheStore.cpp
--- a/Engine/Source/Developer/DerivedDataCache/Private/HttpCacheStore.cpp
+++ b/Engine/Source/Developer/DerivedDataCache/Private/HttpCacheStore.cpp
@@ -15,6 +15,7 @@
 #include "DerivedDataChunk.h"
 #include "DerivedDataRequestOwner.h"
 #include "DerivedDataValue.h"
+#include "DesktopPlatformModule.h"
 #include "Dom/JsonObject.h"
 #include "HAL/IConsoleManager.h"
 #include "Http/HttpClient.h"
@@ -54,6 +55,13 @@
 namespace UE::DerivedData
 {
 
+static bool bHttpEnableOidc = false;
+static FAutoConsoleVariableRef CVarHttpEnableOidc(
+	TEXT("DDC.Http.EnableOidc"),
+	bHttpEnableOidc,
+	TEXT("If true, Oidc tokens are used, otherwise legacy shared credentials are used."),
+	ECVF_Default);
+
 TRACE_DECLARE_INT_COUNTER(HttpDDC_Get, TEXT("HttpDDC Get"));
 TRACE_DECLARE_INT_COUNTER(HttpDDC_GetHit, TEXT("HttpDDC Get Hit"));
 TRACE_DECLARE_INT_COUNTER(HttpDDC_Put, TEXT("HttpDDC Put"));
@@ -110,6 +118,8 @@ struct FHttpCacheStoreParams
 	FString OAuthClientId;
 	FString OAuthSecret;
 	FString OAuthScope;
+	FString OAuthProviderIdentifier;
+	FString OAuthAccessToken;
 	bool bResolveHostCanonicalName = true;
 	bool bReadOnly = false;
 
@@ -176,6 +186,8 @@ public:
 	const FString& GetOAuthClientId() const { return OAuthClientId; }
 	const FString& GetOAuthSecret() const { return OAuthSecret; }
 	const FString& GetOAuthScope() const { return OAuthScope; }
+	const FString& GetOAuthProviderIdentifier() const { return OAuthProviderIdentifier; }
+	const FString& GetOAuthAccessToken() const { return OAuthAccessToken; }
 
 private:
 	FString Domain;
@@ -187,6 +199,9 @@ private:
 	FString OAuthClientId;
 	FString OAuthSecret;
 	FString OAuthScope;
+	FString OAuthProviderIdentifier;
+	FString OAuthAccessToken;
+
 	FCriticalSection AccessCs;
 	FDerivedDataCacheUsageStats UsageStats;
 	FBackendDebugOptions DebugOptions;
@@ -1180,6 +1195,8 @@ FHttpCacheStore::FHttpCacheStore(const FHttpCacheStoreParams& Params)
 	, OAuthClientId(Params.OAuthClientId)
 	, OAuthSecret(Params.OAuthSecret)
 	, OAuthScope(Params.OAuthScope)
+	, OAuthProviderIdentifier(Params.OAuthProviderIdentifier)
+	, OAuthAccessToken(Params.OAuthAccessToken)
 	, Access(nullptr)
 	, bIsUsable(false)
 	, bReadOnly(Params.bReadOnly)
@@ -1301,6 +1318,66 @@ bool FHttpCacheStore::AcquireAccessToken()
 	// get the current serial while we wait to take the CS.
 	const uint32 WantsToUpdateTokenSerial = Access.IsValid() ? Access->GetSerial() : 0u;
 
+	if (bHttpEnableOidc && !OAuthProviderIdentifier.IsEmpty())
+	{
+		FString AccessTokenString;
+		FDateTime TokenExpiresAt;
+		
+		{
+			FScopeLock Lock(&AccessCs);
+
+			// Check if someone has beaten us to update the token, then it 
+			// should now be valid.
+			if (Access.IsValid() && Access->GetSerial() > WantsToUpdateTokenSerial)
+			{
+				return true;
+			}
+
+			if (!OAuthAccessToken.IsEmpty())
+			{
+				if (!Access)
+				{
+					Access = MakeUnique<FHttpAccessToken>();
+				}
+				Access->SetToken(OAuthAccessToken);
+				FailedLoginAttempts = 0;
+				return true;
+			}
+
+			if (FDesktopPlatformModule::Get()->GetOidcAccessToken(FPaths::RootDir(), FPaths::GetProjectFilePath(), OAuthProviderIdentifier, /* Unattended */ true, GWarn, AccessTokenString, TokenExpiresAt))
+			{
+				int32 ExpiryTimeSeconds = (TokenExpiresAt - FDateTime::UtcNow()).GetTotalSeconds();
+				if (!Access)
+				{
+					Access = MakeUnique<FHttpAccessToken>();
+				}
+				Access->SetToken(AccessTokenString);
+				UE_LOG(LogDerivedDataCache, Display, TEXT("OidcToken: Logged in to HTTP DDC services. Expires at %s which is in %d seconds."), *TokenExpiresAt.ToString(), ExpiryTimeSeconds);
+
+				//Schedule a refresh of the token ahead of expiry time (this will not work in commandlets)
+				if (!IsRunningCommandlet())
+				{
+					FTSTicker::GetCoreTicker().AddTicker(FTickerDelegate::CreateLambda(
+						[this](float DeltaTime)
+						{
+							this->AcquireAccessToken();
+							return false;
+						}
+					), ExpiryTimeSeconds - 20.0f);
+				}
+				// Reset failed login attempts, the service is indeed alive.
+				FailedLoginAttempts = 0;
+				return true;
+			}
+			else
+			{
+				UE_LOG(LogDerivedDataCache, Error, TEXT("OidcToken: Failed to log in to HTTP services. "));
+				FailedLoginAttempts++;
+			}
+		}
+		
+	}
+	else // oidc not enabled or not configured
 	{
 		FScopeLock Lock(&AccessCs);
 
@@ -2330,6 +2407,19 @@ void FHttpCacheStoreParams::Parse(const TCHAR* NodeName, const TCHAR* Config)
 
 	FParse::Value(Config, TEXT("OAuthScope="), OAuthScope);
 
+	FParse::Value(Config, TEXT("OAuthProviderIdentifier="), OAuthProviderIdentifier);
+
+	if (FParse::Value(Config, TEXT("OAuthAccessTokenEnvOverride="), OverrideName))
+	{
+		FString AccessToken = FPlatformMisc::GetEnvironmentVariable(*OverrideName);
+		if (!AccessToken.IsEmpty())
+		{
+			OAuthAccessToken = AccessToken;
+			// we do not echo the access token as its sensitive information
+			UE_LOG(LogDerivedDataCache, Log, TEXT("%s: Found oauth access token in %s ."), NodeName, *OverrideName);
+		}
+	}
+
 	// Cache Params
 
 	FParse::Bool(Config, TEXT("ReadOnly="), bReadOnly);
@@ -2419,6 +2509,8 @@ ILegacyCacheStore* GetAnyHttpCacheStore(
 	FString& OutOAuthClientId,
 	FString& OutOAuthSecret,
 	FString& OutOAuthScope,
+	FString& OAuthProviderIdentifier,
+	FString& OAuthAccessToken,
 	FString& OutNamespace,
 	FString& OutStructuredNamespace)
 {
@@ -2430,6 +2522,8 @@ ILegacyCacheStore* GetAnyHttpCacheStore(
 		OutOAuthClientId = HttpBackend->GetOAuthClientId();
 		OutOAuthSecret = HttpBackend->GetOAuthSecret();
 		OutOAuthScope = HttpBackend->GetOAuthScope();
+		OAuthProviderIdentifier = HttpBackend->GetOAuthProviderIdentifier();
+		OAuthAccessToken = HttpBackend->GetOAuthAccessToken();
 		OutNamespace = HttpBackend->GetNamespace();
 		OutStructuredNamespace = HttpBackend->GetStructuredNamespace();
 
diff --git a/Engine/Source/Developer/DerivedDataCache/Private/Tests/HttpCacheStoreTest.cpp b/Engine/Source/Developer/DerivedDataCache/Private/Tests/HttpCacheStoreTest.cpp
--- a/Engine/Source/Developer/DerivedDataCache/Private/Tests/HttpCacheStoreTest.cpp
+++ b/Engine/Source/Developer/DerivedDataCache/Private/Tests/HttpCacheStoreTest.cpp
@@ -44,6 +44,8 @@ ILegacyCacheStore* GetAnyHttpCacheStore(
 	FString& OutOAuthClientId,
 	FString& OutOAuthSecret,
 	FString& OutOAuthScope,
+	FString& OAuthProviderIdentifier,
+	FString& OAuthAccessToken,
 	FString& OutNamespace,
 	FString& OutStructuredNamespace);
 
@@ -158,7 +160,7 @@ protected:
 	static ILegacyCacheStore* GetTestBackend()
 	{
 		static ILegacyCacheStore* CachedBackend = GetAnyHttpCacheStore(
-			TestDomain, TestOAuthProvider, TestOAuthClientId, TestOAuthSecret, TestOAuthScope, TestNamespace, TestStructuredNamespace);
+			TestDomain, TestOAuthProvider, TestOAuthClientId, TestOAuthSecret, TestOAuthScope, TestOAuthProviderIdentifier, TestOAuthAccessToken, TestNamespace, TestStructuredNamespace);
 		return CachedBackend;
 	}
 
@@ -488,6 +490,8 @@ protected:
 	static inline FString TestOAuthClientId;
 	static inline FString TestOAuthSecret;
 	static inline FString TestOAuthScope;
+	static inline FString TestOAuthProviderIdentifier;
+	static inline FString TestOAuthAccessToken;
 	static inline FString TestNamespace;
 	static inline FString TestStructuredNamespace;
 };
diff --git a/Engine/Source/Developer/DesktopPlatform/Private/DesktopPlatformBase.cpp b/Engine/Source/Developer/DesktopPlatform/Private/DesktopPlatformBase.cpp
--- a/Engine/Source/Developer/DesktopPlatform/Private/DesktopPlatformBase.cpp
+++ b/Engine/Source/Developer/DesktopPlatform/Private/DesktopPlatformBase.cpp
@@ -751,6 +751,180 @@ bool FDesktopPlatformBase::RunUnrealBuildTool(const FText& Description, const FS
 	return static_cast<IDesktopPlatform*>(this)->RunUnrealBuildTool(Description, RootDir, Arguments, Warn, ExitCode);
 }
 
+bool FDesktopPlatformBase::GetOidcAccessToken(const FString& RootDir, const FString& ProjectFileName, const FString& ProviderIdentifier, bool Unattended, FFeedbackContext* Warn, FString& OutToken, FDateTime& OutTokenExpiresAt)
+{
+	FString ResultFilePath = FPaths::CreateTempFilename(*FPaths::ProjectIntermediateDir(), TEXT("oidcToken.json"));
+
+	FString Arguments = TEXT(" ");
+	Arguments += FString::Printf(TEXT(" --Service=\"%s\""), *ProviderIdentifier);
+	Arguments += FString::Printf(TEXT(" --OutFile=\"%s\""), *ResultFilePath);
+	if (Unattended)
+	{
+		Arguments += TEXT(" --Unattended=true");
+	}
+
+	// Determine if we need to hint which game project to read configs from
+	if ( !ProjectFileName.IsEmpty() && GetCachedProjectDictionary(RootDir).IsForeignProject(ProjectFileName) )
+	{
+		// Figure out whether it's a foreign project
+		const FUProjectDictionary &ProjectDictionary = GetCachedProjectDictionary(RootDir);
+		if(ProjectDictionary.IsForeignProject(ProjectFileName))
+		{
+			Arguments += FString::Printf(TEXT(" --project=\"%s\""), *IFileManager::Get().ConvertToAbsolutePathForExternalAppForRead(*ProjectFileName));
+		}
+	}
+
+	bool bRes = true;
+	int32 ExitCode;
+	FString ProcessStdout;
+	bRes = InvokeOidcTokenToolSync(LOCTEXT("GetOidcAccessToken", "Fetching OIDC Access Token..."), RootDir, Arguments, Warn, ExitCode, ProcessStdout);
+
+	if (ExitCode == -1337)
+	{
+		UE_LOG(LogDesktopPlatform, Error, TEXT("Unable to allocate a access token. Make sure you are logged in using UGS or using the UGS cli command 'login'. Provider used: '%s'. Ran OidcToken (project file is '%s', exe path is '%s')"), *ProviderIdentifier, *ProjectFileName, *GetOidcTokenExecutableFilename(RootDir));
+		return false;
+	}
+
+	if (!bRes)
+	{
+		UE_LOG(LogDesktopPlatform, Warning, TEXT("Failed to run OidcToken (project file is '%s', exe path is '%s')"), *ProjectFileName, *GetOidcTokenExecutableFilename(RootDir));
+		return false;
+	}
+	
+	// Read the file to a string
+	FString TokenText;
+	if(FFileHelper::LoadFileToString(TokenText, *ResultFilePath))
+	{
+		// deserialize the json file
+		TSharedPtr< FJsonObject > Object;
+		TSharedRef< TJsonReader<> > Reader = TJsonReaderFactory<>::Create(TokenText);
+		if(FJsonSerializer::Deserialize(Reader, Object) && Object.IsValid())
+		{
+			FString Token;
+			FString ExpiresAt;
+			if(Object->TryGetStringField(TEXT("Token"), Token) && Object->TryGetStringField(TEXT("ExpiresAt"), ExpiresAt))
+			{
+				OutToken = Token;
+
+				FDateTime::ParseIso8601(*ExpiresAt, OutTokenExpiresAt);
+
+				return true;
+			}
+		}
+	}
+
+	return false;
+}
+
+
+bool FDesktopPlatformBase::GetOidcTokenStatus(const FString& RootDir, const FString& ProjectFileName, const FString& ProviderIdentifier, FFeedbackContext* Warn, int& OutStatus)
+{
+	FString ResultFilePath = FPaths::CreateTempFilename(*FPaths::ProjectIntermediateDir(), TEXT("oidcToken-status.json"));
+
+	FString Arguments = TEXT(" ");
+	Arguments += FString::Printf(TEXT(" --Service=\"%s\""), *ProviderIdentifier);
+	Arguments += FString::Printf(TEXT(" --OutFile=\"%s\""), *ResultFilePath);
+
+	// Determine if we need to hint which game project to read configs from
+	if ( !ProjectFileName.IsEmpty() && GetCachedProjectDictionary(RootDir).IsForeignProject(ProjectFileName) )
+	{
+		// Figure out whether it's a foreign project
+		const FUProjectDictionary &ProjectDictionary = GetCachedProjectDictionary(RootDir);
+		if(ProjectDictionary.IsForeignProject(ProjectFileName))
+		{
+			Arguments += FString::Printf(TEXT(" --project=\"%s\""), *IFileManager::Get().ConvertToAbsolutePathForExternalAppForRead(*ProjectFileName));
+		}
+	}
+
+	bool bRes = true;
+	int32 ExitCode;
+	FString ProcessStdout;
+	bRes = InvokeOidcTokenToolSync(LOCTEXT("GetOidcAccessToken", "Fetching OIDC Access Token Status..."), RootDir, Arguments, Warn, ExitCode, ProcessStdout);
+
+	if (!bRes)
+	{
+		UE_LOG(LogDesktopPlatform, Warning, TEXT("Failed to run OidcToken to determine token status (project file is '%s', exe path is '%s')"), *ProjectFileName, *GetOidcTokenExecutableFilename(RootDir));
+		return false;
+	}
+	
+	// Read the file to a string
+	FString TokenText;
+	if(FFileHelper::LoadFileToString(TokenText, *ResultFilePath))
+	{
+		// deserialize the json file
+		TSharedPtr< FJsonObject > Object;
+		TSharedRef< TJsonReader<> > Reader = TJsonReaderFactory<>::Create(TokenText);
+		if(FJsonSerializer::Deserialize(Reader, Object) && Object.IsValid())
+		{
+			int Status;
+			if(Object->TryGetNumberField(TEXT("Status"), Status))
+			{
+				OutStatus = Status;
+
+				return true;
+			}
+		}
+	}
+
+	return false;
+}
+
+bool FDesktopPlatformBase::InvokeOidcTokenToolSync(const FText& Description, const FString& RootDir, const FString& Arguments, FFeedbackContext* Warn, int32& OutReturnCode, FString& OutProcOutput)
+{
+	TRACE_CPUPROFILER_EVENT_SCOPE(FDesktopPlatformBase::InvokeOidcTokenToolSync);
+	OutReturnCode = 1;
+
+	void* PipeRead = nullptr;
+	void* PipeWrite = nullptr;
+
+	verify(FPlatformProcess::CreatePipe(PipeRead, PipeWrite));
+
+	bool bInvoked = false;
+	FProcHandle ProcHandle = InvokeOidcTokenToolAsync(Arguments, PipeRead, PipeWrite);
+	if (ProcHandle.IsValid())
+	{
+		// rather than waiting, we must flush the read pipe or UBT will stall if it writes out a ton of text to the console.
+		while (FPlatformProcess::IsProcRunning(ProcHandle))
+		{
+			OutProcOutput += FPlatformProcess::ReadPipe(PipeRead);
+			FPlatformProcess::Sleep(0.1f);
+		}		
+		bInvoked = true;
+		bool bGotReturnCode = FPlatformProcess::GetProcReturnCode(ProcHandle, &OutReturnCode);		
+		check(bGotReturnCode);
+	}
+	else
+	{
+		bInvoked = false;
+		OutReturnCode = -1;
+		OutProcOutput = TEXT("");
+	}
+
+
+	FPlatformProcess::ClosePipe(PipeRead, PipeWrite);
+
+	return bInvoked;
+}
+
+FProcHandle FDesktopPlatformBase::InvokeOidcTokenToolAsync(const FString& InArguments, void*& OutReadPipe, void*& OutWritePipe)
+{
+	FString CmdLineParams = InArguments;
+	FString ExecutableFileName = GetOidcTokenExecutableFilename(FPaths::RootDir());
+	UE_LOG(LogDesktopPlatform, Display, TEXT("Launching OidcToken... [%s %s]"), *ExecutableFileName, *CmdLineParams);
+
+	const bool bLaunchDetached = false;
+	const bool bLaunchHidden = true;
+	const bool bLaunchReallyHidden = bLaunchHidden;
+
+	FProcHandle ProcHandle = FPlatformProcess::CreateProc(*ExecutableFileName, *CmdLineParams, bLaunchDetached, bLaunchHidden, bLaunchReallyHidden, NULL, 0, NULL, OutWritePipe, OutReadPipe);
+	if (!ProcHandle.IsValid())
+	{
+		UE_LOG(LogDesktopPlatform, Warning, TEXT("Failed to launch OidcToken (exe path is '%s')"), *ExecutableFileName);
+	}
+
+	return ProcHandle;
+}
+
 struct FTargetFileVisitor : IPlatformFile::FDirectoryStatVisitor
 {
 	TSet<FString>& RemainingTargetNames;
diff --git a/Engine/Source/Developer/DesktopPlatform/Private/DesktopPlatformBase.h b/Engine/Source/Developer/DesktopPlatform/Private/DesktopPlatformBase.h
--- a/Engine/Source/Developer/DesktopPlatform/Private/DesktopPlatformBase.h
+++ b/Engine/Source/Developer/DesktopPlatform/Private/DesktopPlatformBase.h
@@ -54,6 +54,9 @@ public:
 	virtual FProcHandle InvokeUnrealBuildToolAsync(const FString& InCmdLineParams, FOutputDevice &Ar, void*& OutReadPipe, void*& OutWritePipe, bool bSkipBuildUBT = false) override;
 	virtual bool RunUnrealBuildTool(const FText& Description, const FString& RootDir, const FString& Arguments, FFeedbackContext* Warn) override final;
 
+    virtual bool GetOidcAccessToken(const FString& RootDir, const FString& ProjectFileName, const FString& ProviderIdentifier, bool Unattended, FFeedbackContext* Warn, FString& OutToken, FDateTime& OutTokenExpiresAt) override;
+	virtual bool GetOidcTokenStatus(const FString& RootDir, const FString& ProjectFileName, const FString& ProviderIdentifier, FFeedbackContext* Warn, int& OutStatus) override;
+
 	virtual const TArray<FTargetInfo>& GetTargetsForProject(const FString& ProjectFile) const override;
 	virtual const TArray<FTargetInfo>& GetTargetsForCurrentProject() const override;
 
@@ -79,6 +82,9 @@ private:
 	void GetProjectBuildProducts(const FString& ProjectFileName, TArray<FString> &OutFileNames, TArray<FString> &OutDirectoryNames);
 	bool BuildUnrealBuildTool(const FString& RootDir, FOutputDevice &Ar);
 
+	bool InvokeOidcTokenToolSync(const FText& Description, const FString& RootDir, const FString& Arguments, FFeedbackContext* Warn, int32& OutReturnCode, FString& OutProcOutput);
+	FProcHandle InvokeOidcTokenToolAsync(const FString& InArguments, void*& OutReadPipe, void*& OutWritePipe);
+
 	static bool ReadTargetInfo(const FString& FileName, TArray<FTargetInfo>& Targets);
 
 	FString GetUserDir(const FString& Identifier);
@@ -89,6 +95,7 @@ protected:
 
 	FString GetUnrealBuildToolProjectFileName(const FString& RootDir) const;
 	FString GetUnrealBuildToolExecutableFilename(const FString& RootDir) const;	
+	virtual FString GetOidcTokenExecutableFilename(const FString& RootDir) const = 0;
 };
 
 PRAGMA_ENABLE_DEPRECATION_WARNINGS
diff --git a/Engine/Source/Developer/DesktopPlatform/Private/DesktopPlatformStub.h b/Engine/Source/Developer/DesktopPlatform/Private/DesktopPlatformStub.h
--- a/Engine/Source/Developer/DesktopPlatform/Private/DesktopPlatformStub.h
+++ b/Engine/Source/Developer/DesktopPlatform/Private/DesktopPlatformStub.h
@@ -26,6 +26,8 @@ public:
 	virtual bool RunUnrealBuildTool(const FText& Description, const FString& RootDir, const FString& Arguments, FFeedbackContext* Warn, int32& OutExitCode) override { OutExitCode = 1;  return false; }
 	virtual bool IsUnrealBuildToolRunning() override { return false; }
 
+	virtual FString GetOidcTokenExecutableFilename(const FString& RootDir) const override { return ""; }
+
 	virtual FFeedbackContext* GetNativeFeedbackContext() override { return nullptr; }
 
 	virtual FString GetUserTempPath() override { return FPaths::ProjectIntermediateDir() + TEXT("UserTemp/"); }
diff --git a/Engine/Source/Developer/DesktopPlatform/Private/Linux/DesktopPlatformLinux.cpp b/Engine/Source/Developer/DesktopPlatform/Private/Linux/DesktopPlatformLinux.cpp
--- a/Engine/Source/Developer/DesktopPlatform/Private/Linux/DesktopPlatformLinux.cpp
+++ b/Engine/Source/Developer/DesktopPlatform/Private/Linux/DesktopPlatformLinux.cpp
@@ -514,4 +514,9 @@ FString FDesktopPlatformLinux::GetUserTempPath()
 	return FString(FPlatformProcess::UserTempDir());
 }
 
+FString FDesktopPlatformLinux::GetOidcTokenExecutableFilename(const FString& RootDir) const
+{	
+	return FPaths::ConvertRelativePathToFull(RootDir / TEXT("Engine/Binaries/DotNET/OidcToken/linux-x64/OidcToken.exe"));
+}
+
 #undef LOCTEXT_NAMESPACE
diff --git a/Engine/Source/Developer/DesktopPlatform/Private/Linux/DesktopPlatformLinux.h b/Engine/Source/Developer/DesktopPlatform/Private/Linux/DesktopPlatformLinux.h
--- a/Engine/Source/Developer/DesktopPlatform/Private/Linux/DesktopPlatformLinux.h
+++ b/Engine/Source/Developer/DesktopPlatform/Private/Linux/DesktopPlatformLinux.h
@@ -36,6 +36,9 @@ public:
 	FDesktopPlatformLinux();
 	virtual ~FDesktopPlatformLinux();
 
+protected:
+	virtual FString GetOidcTokenExecutableFilename(const FString& RootDir) const override;
+
 private:
 	bool FileDialogShared(bool bSave, const void* ParentWindowHandle, const FString& DialogTitle, const FString& DefaultPath, const FString& DefaultFile, const FString& FileTypes, uint32 Flags, TArray<FString>& OutFilenames, int32& OutFilterIndex);
 };
diff --git a/Engine/Source/Developer/DesktopPlatform/Private/Mac/DesktopPlatformMac.cpp b/Engine/Source/Developer/DesktopPlatform/Private/Mac/DesktopPlatformMac.cpp
--- a/Engine/Source/Developer/DesktopPlatform/Private/Mac/DesktopPlatformMac.cpp
+++ b/Engine/Source/Developer/DesktopPlatform/Private/Mac/DesktopPlatformMac.cpp
@@ -690,4 +690,8 @@ FString FDesktopPlatformMac::GetUserTempPath()
 	return FString(FPlatformProcess::UserTempDir());
 }
 
+FString FDesktopPlatformMac::GetOidcTokenExecutableFilename(const FString& RootDir) const
+{
+	return FPaths::ConvertRelativePathToFull(RootDir / TEXT("Engine/Binaries/DotNET/OidcToken/mac-x64/OidcToken.exe"));
+}
 #undef LOCTEXT_NAMESPACE
diff --git a/Engine/Source/Developer/DesktopPlatform/Private/Mac/DesktopPlatformMac.h b/Engine/Source/Developer/DesktopPlatform/Private/Mac/DesktopPlatformMac.h
--- a/Engine/Source/Developer/DesktopPlatform/Private/Mac/DesktopPlatformMac.h
+++ b/Engine/Source/Developer/DesktopPlatform/Private/Mac/DesktopPlatformMac.h
@@ -28,6 +28,9 @@ public:
 
 	virtual FString GetUserTempPath() override;
 
+protected:
+	virtual FString GetOidcTokenExecutableFilename(const FString& RootDir) const override;
+
 private:
 	bool FileDialogShared(bool bSave, const void* ParentWindowHandle, const FString& DialogTitle, const FString& DefaultPath, const FString& DefaultFile, const FString& FileTypes, uint32 Flags, TArray<FString>& OutFilenames, int32& OutFilterIndex);
 };
diff --git a/Engine/Source/Developer/DesktopPlatform/Private/Windows/DesktopPlatformWindows.cpp b/Engine/Source/Developer/DesktopPlatform/Private/Windows/DesktopPlatformWindows.cpp
--- a/Engine/Source/Developer/DesktopPlatform/Private/Windows/DesktopPlatformWindows.cpp
+++ b/Engine/Source/Developer/DesktopPlatform/Private/Windows/DesktopPlatformWindows.cpp
@@ -661,4 +661,9 @@ int32 FDesktopPlatformWindows::GetShellIntegrationVersion(const FString &FileNam
 	return 0;
 }
 
+FString FDesktopPlatformWindows::GetOidcTokenExecutableFilename(const FString& RootDir) const
+{
+	return FPaths::ConvertRelativePathToFull(RootDir / TEXT("Engine/Binaries/DotNET/OidcToken/win-x64/OidcToken.exe"));
+}
+
 #undef LOCTEXT_NAMESPACE
diff --git a/Engine/Source/Developer/DesktopPlatform/Private/Windows/DesktopPlatformWindows.h b/Engine/Source/Developer/DesktopPlatform/Private/Windows/DesktopPlatformWindows.h
--- a/Engine/Source/Developer/DesktopPlatform/Private/Windows/DesktopPlatformWindows.h
+++ b/Engine/Source/Developer/DesktopPlatform/Private/Windows/DesktopPlatformWindows.h
@@ -36,6 +36,9 @@ public:
 
 	virtual FString GetUserTempPath() override;
 
+protected:
+	virtual FString GetOidcTokenExecutableFilename(const FString& RootDir) const override;
+
 private:
 	bool FileDialogShared(bool bSave, const void* ParentWindowHandle, const FString& DialogTitle, const FString& DefaultPath, const FString& DefaultFile, const FString& FileTypes, uint32 Flags, TArray<FString>& OutFilenames, int32& OutFilterIndex);
 	void GetRequiredRegistrySettings(TIndirectArray<FRegistryRootedKey> &RootedKeys);
diff --git a/Engine/Source/Developer/DesktopPlatform/Public/IDesktopPlatform.h b/Engine/Source/Developer/DesktopPlatform/Public/IDesktopPlatform.h
--- a/Engine/Source/Developer/DesktopPlatform/Public/IDesktopPlatform.h
+++ b/Engine/Source/Developer/DesktopPlatform/Public/IDesktopPlatform.h
@@ -456,4 +456,30 @@ public:
 	* @return Path to the folder
 	*/
 	virtual FString GetDefaultProjectCreationPath() = 0;
+
+	/**
+	* Get a Access token for use by OIDC compliant Identity Providers
+
+	* @param RootDir			Engine root directory for the project to use.
+	* @param ProjectFileName	Filename of the current project
+	* @param ProviderIdentifier	The identifier for the provider as configured in oidc-token.json
+	* @param Unattended			True to indicate that no user interaction should be assumed
+	* @param Warn				Feedback context to use for progress updates
+	* @param OutToken			The allocated access token
+	* @param OutTokenExpiresAt	When the token expires
+	* @return true if the task completed successfully.
+	*/
+	virtual bool GetOidcAccessToken(const FString& RootDir, const FString& ProjectFileName, const FString& ProviderIdentifier, bool Unattended, FFeedbackContext* Warn, FString& OutToken, FDateTime& OutTokenExpiresAt) = 0;
+
+	/**
+	* Get the status of a access token for use by OIDC compliant Identity Providers
+
+	* @param RootDir			Engine root directory for the project to use.
+	* @param ProjectFileName	Filename of the current project
+	* @param ProviderIdentifier	The identifier for the provider as configured in oidc-token.json
+	* @param Warn				Feedback context to use for progress updates
+	* @param OutStatus			The allocated access token
+	* @return true if the task completed successfully.
+	*/
+	virtual bool GetOidcTokenStatus(const FString& RootDir, const FString& ProjectFileName, const FString& ProviderIdentifier, FFeedbackContext* Warn, int& OutStatus) = 0;
 };
