Fix OidcToken status output localization key

diff --git a/Engine/Source/Developer/DesktopPlatform/Private/DesktopPlatformBase.cpp b/Engine/Source/Developer/DesktopPlatform/Private/DesktopPlatformBase.cpp
--- a/Engine/Source/Developer/DesktopPlatform/Private/DesktopPlatformBase.cpp
+++ b/Engine/Source/Developer/DesktopPlatform/Private/DesktopPlatformBase.cpp
@@ -839,7 +839,7 @@ bool FDesktopPlatformBase::GetOidcTokenStatus(const FString& RootDir, const FStr
 	bool bRes = true;
 	int32 ExitCode;
 	FString ProcessStdout;
-	bRes = InvokeOidcTokenToolSync(LOCTEXT("GetOidcAccessToken", "Fetching OIDC Access Token Status..."), RootDir, Arguments, Warn, ExitCode, ProcessStdout);
+	bRes = InvokeOidcTokenToolSync(LOCTEXT("GetOidcAccessTokenStatus", "Fetching OIDC Access Token Status..."), RootDir, Arguments, Warn, ExitCode, ProcessStdout);
 
 	if (!bRes)
 	{
