Allow users to interactivly login if they are missing a token. Will present a dialog for users if not unattended so they are aware they need to login.

diff --git a/Engine/Source/Developer/DerivedDataCache/Private/HttpCacheStore.cpp b/Engine/Source/Developer/DerivedDataCache/Private/HttpCacheStore.cpp
--- a/Engine/Source/Developer/DerivedDataCache/Private/HttpCacheStore.cpp
+++ b/Engine/Source/Developer/DerivedDataCache/Private/HttpCacheStore.cpp
@@ -1672,7 +1672,7 @@ bool FHttpCacheStore::AcquireAccessToken(IHttpClient* Client)
 	{
 		FString AccessTokenString;
 		FDateTime TokenExpiresAt;
-		if (FDesktopPlatformModule::Get()->GetOidcAccessToken(FPaths::RootDir(), FPaths::GetProjectFilePath(), OAuthProviderIdentifier, /* Unattended */ true, GWarn, AccessTokenString, TokenExpiresAt))
+		if (FDesktopPlatformModule::Get()->GetOidcAccessToken(FPaths::RootDir(), FPaths::GetProjectFilePath(), OAuthProviderIdentifier, FApp::IsUnattended(), GWarn, AccessTokenString, TokenExpiresAt))
 		{
 			const double ExpiryTimeSeconds = (TokenExpiresAt - FDateTime::UtcNow()).GetTotalSeconds();
 			UE_LOG(LogDerivedDataCache, Display,
diff --git a/Engine/Source/Developer/DesktopPlatform/Private/DesktopPlatformBase.cpp b/Engine/Source/Developer/DesktopPlatform/Private/DesktopPlatformBase.cpp
--- a/Engine/Source/Developer/DesktopPlatform/Private/DesktopPlatformBase.cpp
+++ b/Engine/Source/Developer/DesktopPlatform/Private/DesktopPlatformBase.cpp
@@ -15,7 +15,7 @@
 #include "Modules/ModuleManager.h"
 #include "DesktopPlatformPrivate.h"
 #include "Misc/OutputDeviceRedirector.h"
-
+#include "Misc/ScopedSlowTask.h"
 
 #define LOCTEXT_NAMESPACE "DesktopPlatform"
 
@@ -758,10 +758,6 @@ bool FDesktopPlatformBase::GetOidcAccessToken(const FString& RootDir, const FStr
 	FString Arguments = TEXT(" ");
 	Arguments += FString::Printf(TEXT(" --Service=\"%s\""), *ProviderIdentifier);
 	Arguments += FString::Printf(TEXT(" --OutFile=\"%s\""), *ResultFilePath);
-	if (Unattended)
-	{
-		Arguments += TEXT(" --Unattended=true");
-	}
 
 	// Determine if we need to hint which game project to read configs from
 	if ( !ProjectFileName.IsEmpty() && GetCachedProjectDictionary(RootDir).IsForeignProject(ProjectFileName) )
@@ -773,16 +769,31 @@ bool FDesktopPlatformBase::GetOidcAccessToken(const FString& RootDir, const FStr
 			Arguments += FString::Printf(TEXT(" --project=\"%s\""), *IFileManager::Get().ConvertToAbsolutePathForExternalAppForRead(*ProjectFileName));
 		}
 	}
+	FString UnattendedArguments = Arguments;
+	UnattendedArguments += TEXT(" --Unattended=true");
 
+	// first we attempt to fetch a token using cached offline tokens, thus setting unattended
 	bool bRes = true;
 	int32 ExitCode;
 	FString ProcessStdout;
-	bRes = InvokeOidcTokenToolSync(LOCTEXT("GetOidcAccessToken", "Fetching OIDC Access Token..."), RootDir, Arguments, Warn, ExitCode, ProcessStdout);
+	bRes = InvokeOidcTokenToolSync(LOCTEXT("GetOidcAccessToken", "Fetching OIDC Access Token..."), RootDir, UnattendedArguments, Warn, ExitCode, ProcessStdout);
 
 	if (ExitCode == -1337)
 	{
-		UE_LOG(LogDesktopPlatform, Error, TEXT("Unable to allocate a access token. Make sure you are logged in using UGS or using the UGS cli command 'login'. Provider used: '%s'. Ran OidcToken (project file is '%s', exe path is '%s')"), *ProviderIdentifier, *ProjectFileName, *GetOidcTokenExecutableFilename(RootDir));
-		return false;
+		if (!Unattended)
+		{
+			bRes = GetOidcAccessTokenInteractive(RootDir, Arguments, Warn, ExitCode);
+			if (!bRes)
+			{
+				UE_LOG(LogDesktopPlatform, Error, TEXT("Unable to allocate an access token. Interactive login failed, make sure you are logged in using UGS or using the UGS cli command 'login'. Provider used: '%s'. Ran OidcToken (project file is '%s', exe path is '%s')"), *ProviderIdentifier, *ProjectFileName, *GetOidcTokenExecutableFilename(RootDir));
+				return false;
+			}
+		}
+		else
+		{
+			UE_LOG(LogDesktopPlatform, Error, TEXT("Unable to allocate an access token. Unattended set so unable to request interactive login. Make sure you are logged in using UGS or using the UGS cli command 'login'. Provider used: '%s'. Ran OidcToken (project file is '%s', exe path is '%s')"), *ProviderIdentifier, *ProjectFileName, *GetOidcTokenExecutableFilename(RootDir));
+			return false;
+		}
 	}
 
 	if (!bRes)
@@ -869,6 +880,84 @@ bool FDesktopPlatformBase::GetOidcTokenStatus(const FString& RootDir, const FStr
 	return false;
 }
 
+bool FDesktopPlatformBase::GetOidcAccessTokenInteractive(const FString& RootDir,  const FString& Arguments, FFeedbackContext* Warn, int32& OutReturnCode)
+{
+	FText OidcInteractivePromptTitle = NSLOCTEXT("OidcToken", "OidcToken_InteractiveLaunchPromptTitle", "User Login");
+	FText OidcInteractiveLaunchPromptText = NSLOCTEXT("OidcToken", "OidcToken_InteractiveLaunch", "Your system browser will now be opened for you to finish the login process.");
+	FPlatformMisc::MessageBoxExt(EAppMsgType::Ok, *OidcInteractiveLaunchPromptText.ToString(), *OidcInteractivePromptTitle.ToString());
+	// user has acknowledged the login, we update the editor progress and then we run oidc token to spawn the browser and prompt the login
+
+	FScopedSlowTask WaitForOidcTokenSlowTask(0, NSLOCTEXT("OidcToken", "OidcToken_WaitingForToken", "Waiting for OidcToken to finish login"));
+
+	// run the oidc token app and wait for it to finish, prompting users if they have not logged in after a while
+	OutReturnCode = 1;
+	void* PipeRead = nullptr;
+	void* PipeWrite = nullptr;
+
+	verify(FPlatformProcess::CreatePipe(PipeRead, PipeWrite));
+
+	bool bInvoked = false;
+	FProcHandle ProcHandle = InvokeOidcTokenToolAsync(Arguments, PipeRead, PipeWrite);
+
+	if (ProcHandle.IsValid())
+	{
+		bInvoked = true;
+	}
+
+	uint64 WaitStartTime = FPlatformTime::Cycles64();
+	enum class EWaitDurationPhase
+	{
+		Initial,
+		Prompt,
+		Waiting
+	} DurationPhase = EWaitDurationPhase::Initial;
+	bool bIsFinished = false;
+	while (!bIsFinished)
+	{
+		// check if token app has finished running
+		bIsFinished = !FPlatformProcess::IsProcRunning(ProcHandle);
+
+		double WaitDuration = FPlatformTime::ToSeconds64(FPlatformTime::Cycles64() - WaitStartTime);
+		if (DurationPhase == EWaitDurationPhase::Initial)
+		{
+			// Note that the dialog may not show up when tokens are allocated early in the launch cycle, but this will at least ensure
+			// the splash screen is refreshed with the appropriate text status message.
+			WaitForOidcTokenSlowTask.MakeDialog(true, false);
+			UE_LOG(LogDesktopPlatform, Display, TEXT("Waiting for OidcToken to finish login..."));
+			DurationPhase = EWaitDurationPhase::Prompt;
+		}
+		// once we have waited for 10 seconds without success we give the user a option to abort
+		else if (WaitDuration > 10.0 && DurationPhase == EWaitDurationPhase::Prompt)
+		{
+			FText OidcLongWaitPromptTitle = NSLOCTEXT("OidcToken", "OidcToken_LongWaitPromptTitle", "Wait for user login?");
+			FText OidcLongWaitPromptText = NSLOCTEXT("OidcToken", "OidcToken_LongWaitPromptText", "Login is taking a long time, make sure you have entered your credentials in your browser window. It can be in a tab in an already existing window. Keep waiting?");
+			if (FPlatformMisc::MessageBoxExt(EAppMsgType::YesNo, *OidcLongWaitPromptText.ToString(), *OidcLongWaitPromptTitle.ToString()) == EAppReturnType::No)
+			{
+				break;
+			}
+			// change phase so we do not prompt the user again
+			DurationPhase = EWaitDurationPhase::Waiting;
+		}
+
+		if (WaitForOidcTokenSlowTask.ShouldCancel())
+		{
+			break;
+		}
+		FPlatformProcess::Sleep(0.1f);
+	}
+
+	bool bGotReturnCode = FPlatformProcess::GetProcReturnCode(ProcHandle, &OutReturnCode);		
+	check(bGotReturnCode);
+	FPlatformProcess::ClosePipe(PipeRead, PipeWrite);
+	
+	if (!bIsFinished)
+	{
+		OutReturnCode = -1;
+	}
+
+	return bInvoked;
+}
+
 bool FDesktopPlatformBase::InvokeOidcTokenToolSync(const FText& Description, const FString& RootDir, const FString& Arguments, FFeedbackContext* Warn, int32& OutReturnCode, FString& OutProcOutput)
 {
 	TRACE_CPUPROFILER_EVENT_SCOPE(FDesktopPlatformBase::InvokeOidcTokenToolSync);
diff --git a/Engine/Source/Developer/DesktopPlatform/Private/DesktopPlatformBase.h b/Engine/Source/Developer/DesktopPlatform/Private/DesktopPlatformBase.h
--- a/Engine/Source/Developer/DesktopPlatform/Private/DesktopPlatformBase.h
+++ b/Engine/Source/Developer/DesktopPlatform/Private/DesktopPlatformBase.h
@@ -84,6 +84,7 @@ private:
 
 	bool InvokeOidcTokenToolSync(const FText& Description, const FString& RootDir, const FString& Arguments, FFeedbackContext* Warn, int32& OutReturnCode, FString& OutProcOutput);
 	FProcHandle InvokeOidcTokenToolAsync(const FString& InArguments, void*& OutReadPipe, void*& OutWritePipe);
+	bool GetOidcAccessTokenInteractive(const FString& RootDir,  const FString& Arguments, FFeedbackContext* Warn, int32& OutReturnCode);
 
 	static bool ReadTargetInfo(const FString& FileName, TArray<FTargetInfo>& Targets);
 
