OidcToken - Fixed path passed for project to oidc token to always be the project directory (no matter if its a foreign project or not). This makes sure we are able to read the configuration from the project if a project is set.

Also updated the prompt asking user to login to make it clear who is asking and for what (assuming this is used by DDC - which is only the case now but may not be forever)

diff --git a/Engine/Source/Developer/DesktopPlatform/Private/DesktopPlatformBase.cpp b/Engine/Source/Developer/DesktopPlatform/Private/DesktopPlatformBase.cpp
--- a/Engine/Source/Developer/DesktopPlatform/Private/DesktopPlatformBase.cpp
+++ b/Engine/Source/Developer/DesktopPlatform/Private/DesktopPlatformBase.cpp
@@ -758,17 +758,7 @@ bool FDesktopPlatformBase::GetOidcAccessToken(const FString& RootDir, const FStr
 	FString Arguments = TEXT(" ");
 	Arguments += FString::Printf(TEXT(" --Service=\"%s\""), *ProviderIdentifier);
 	Arguments += FString::Printf(TEXT(" --OutFile=\"%s\""), *ResultFilePath);
-
-	// Determine if we need to hint which game project to read configs from
-	if ( !ProjectFileName.IsEmpty() && GetCachedProjectDictionary(RootDir).IsForeignProject(ProjectFileName) )
-	{
-		// Figure out whether it's a foreign project
-		const FUProjectDictionary &ProjectDictionary = GetCachedProjectDictionary(RootDir);
-		if(ProjectDictionary.IsForeignProject(ProjectFileName))
-		{
-			Arguments += FString::Printf(TEXT(" --project=\"%s\""), *IFileManager::Get().ConvertToAbsolutePathForExternalAppForRead(*ProjectFileName));
-		}
-	}
+	Arguments += FString::Printf(TEXT(" --project=\"%s\""),  *IFileManager::Get().ConvertToAbsolutePathForExternalAppForRead(*FPaths::ProjectDir()));
 	FString UnattendedArguments = Arguments;
 	UnattendedArguments += TEXT(" --Unattended=true");
 
@@ -835,17 +825,7 @@ bool FDesktopPlatformBase::GetOidcTokenStatus(const FString& RootDir, const FStr
 	FString Arguments = TEXT(" ");
 	Arguments += FString::Printf(TEXT(" --Service=\"%s\""), *ProviderIdentifier);
 	Arguments += FString::Printf(TEXT(" --OutFile=\"%s\""), *ResultFilePath);
-
-	// Determine if we need to hint which game project to read configs from
-	if ( !ProjectFileName.IsEmpty() && GetCachedProjectDictionary(RootDir).IsForeignProject(ProjectFileName) )
-	{
-		// Figure out whether it's a foreign project
-		const FUProjectDictionary &ProjectDictionary = GetCachedProjectDictionary(RootDir);
-		if(ProjectDictionary.IsForeignProject(ProjectFileName))
-		{
-			Arguments += FString::Printf(TEXT(" --project=\"%s\""), *IFileManager::Get().ConvertToAbsolutePathForExternalAppForRead(*ProjectFileName));
-		}
-	}
+	Arguments += FString::Printf(TEXT(" --project=\"%s\""),  *IFileManager::Get().ConvertToAbsolutePathForExternalAppForRead(*FPaths::ProjectDir()));
 
 	bool bRes = true;
 	int32 ExitCode;
@@ -882,8 +862,8 @@ bool FDesktopPlatformBase::GetOidcTokenStatus(const FString& RootDir, const FStr
 
 bool FDesktopPlatformBase::GetOidcAccessTokenInteractive(const FString& RootDir,  const FString& Arguments, FFeedbackContext* Warn, int32& OutReturnCode)
 {
-	FText OidcInteractivePromptTitle = NSLOCTEXT("OidcToken", "OidcToken_InteractiveLaunchPromptTitle", "User Login");
-	FText OidcInteractiveLaunchPromptText = NSLOCTEXT("OidcToken", "OidcToken_InteractiveLaunch", "Your system browser will now be opened for you to finish the login process.");
+	FText OidcInteractivePromptTitle = NSLOCTEXT("OidcToken", "OidcToken_InteractiveLaunchPromptTitle", "Unreal Engine - User Login");
+	FText OidcInteractiveLaunchPromptText = NSLOCTEXT("OidcToken", "OidcToken_InteractiveLaunch", "Authentication is required for Http Derived Data Cache to work. Your system browser will now be opened for you to finish the login process.");
 	FPlatformMisc::MessageBoxExt(EAppMsgType::Ok, *OidcInteractiveLaunchPromptText.ToString(), *OidcInteractivePromptTitle.ToString());
 	// user has acknowledged the login, we update the editor progress and then we run oidc token to spawn the browser and prompt the login
 
