Removed duplicated code from a 5.0 patch fix.

#rnx

diff --git a/Engine/Source/Developer/DesktopPlatform/Private/DesktopPlatformBase.cpp b/Engine/Source/Developer/DesktopPlatform/Private/DesktopPlatformBase.cpp
--- a/Engine/Source/Developer/DesktopPlatform/Private/DesktopPlatformBase.cpp
+++ b/Engine/Source/Developer/DesktopPlatform/Private/DesktopPlatformBase.cpp
@@ -12,10 +12,12 @@
 #include "Serialization/JsonTypes.h"
 #include "Serialization/JsonReader.h"
 #include "Serialization/JsonSerializer.h"
+#include "String/LexFromString.h"
 #include "Modules/ModuleManager.h"
 #include "DesktopPlatformPrivate.h"
 #include "Misc/OutputDeviceRedirector.h"
 #include "Misc/ScopedSlowTask.h"
+#include "Misc/SecureHash.h"
 
 #define LOCTEXT_NAMESPACE "DesktopPlatform"
 
@@ -751,6 +753,58 @@ bool FDesktopPlatformBase::RunUnrealBuildTool(const FText& Description, const FS
 	return static_cast<IDesktopPlatform*>(this)->RunUnrealBuildTool(Description, RootDir, Arguments, Warn, ExitCode);
 }
 
+bool FDesktopPlatformBase::IsUnrealBuildToolRunning()
+{
+	FString RunsDir = FPaths::Combine(FPaths::EngineIntermediateDir(), TEXT("UbtRuns"));
+	if (!FPaths::DirectoryExists(RunsDir))
+	{
+		return false;
+	}
+
+	bool bIsRunning = false;
+	IFileManager::Get().IterateDirectory(*RunsDir, [&bIsRunning](const TCHAR* Pathname, bool bIsDirectory)
+		{
+			if (!bIsDirectory)
+			{
+				bool bDeleteFile = true;
+
+				FString Filename = FPaths::GetBaseFilename(FString(Pathname));
+				const TCHAR* Delim = FCString::Strchr(*Filename, '_');
+				if (Delim != nullptr)
+				{
+					FStringView Pid(*Filename, UE_PTRDIFF_TO_INT32(Delim - *Filename));
+					int ProcessId = 0;
+					LexFromString(ProcessId, Pid);
+					FString EntryFullPath = FPlatformProcess::GetApplicationName(ProcessId);
+					if (!EntryFullPath.IsEmpty())
+					{
+						EntryFullPath.ToUpperInline();
+						const auto Utf8String = StringCast<UTF8CHAR>(*EntryFullPath);
+						FMD5Hash Hash;
+						LexFromString(Hash, Delim + 1);
+
+						FMD5 Md5Gen;
+						Md5Gen.Update(reinterpret_cast<const uint8*>(Utf8String.Get()), Utf8String.Length());
+						FMD5Hash TestHash;
+						TestHash.Set(Md5Gen);
+						if (Hash == TestHash)
+						{
+							bDeleteFile = false;
+							bIsRunning = true;
+						}
+					}
+					if (bDeleteFile)
+					{
+						IFileManager::Get().Delete(Pathname);
+					}
+				}
+			}
+			return true;
+		});
+
+	return bIsRunning;
+}
+
 bool FDesktopPlatformBase::GetOidcAccessToken(const FString& RootDir, const FString& ProjectFileName, const FString& ProviderIdentifier, bool Unattended, FFeedbackContext* Warn, FString& OutToken, FDateTime& OutTokenExpiresAt)
 {
 	FString ResultFilePath = FPaths::CreateTempFilename(*FPaths::ProjectIntermediateDir(), TEXT("oidcToken.json"));
diff --git a/Engine/Source/Developer/DesktopPlatform/Private/DesktopPlatformBase.h b/Engine/Source/Developer/DesktopPlatform/Private/DesktopPlatformBase.h
--- a/Engine/Source/Developer/DesktopPlatform/Private/DesktopPlatformBase.h
+++ b/Engine/Source/Developer/DesktopPlatform/Private/DesktopPlatformBase.h
@@ -53,6 +53,7 @@ public:
 	virtual bool InvokeUnrealBuildToolSync(const FString& InCmdLineParams, FOutputDevice &Ar, bool bSkipBuildUBT, int32& OutReturnCode, FString& OutProcOutput) override;
 	virtual FProcHandle InvokeUnrealBuildToolAsync(const FString& InCmdLineParams, FOutputDevice &Ar, void*& OutReadPipe, void*& OutWritePipe, bool bSkipBuildUBT = false) override;
 	virtual bool RunUnrealBuildTool(const FText& Description, const FString& RootDir, const FString& Arguments, FFeedbackContext* Warn) override final;
+	virtual bool IsUnrealBuildToolRunning() override;
 
     virtual bool GetOidcAccessToken(const FString& RootDir, const FString& ProjectFileName, const FString& ProviderIdentifier, bool Unattended, FFeedbackContext* Warn, FString& OutToken, FDateTime& OutTokenExpiresAt) override;
 	virtual bool GetOidcTokenStatus(const FString& RootDir, const FString& ProjectFileName, const FString& ProviderIdentifier, FFeedbackContext* Warn, int& OutStatus) override;
diff --git a/Engine/Source/Developer/DesktopPlatform/Private/Linux/DesktopPlatformLinux.cpp b/Engine/Source/Developer/DesktopPlatform/Private/Linux/DesktopPlatformLinux.cpp
--- a/Engine/Source/Developer/DesktopPlatform/Private/Linux/DesktopPlatformLinux.cpp
+++ b/Engine/Source/Developer/DesktopPlatform/Private/Linux/DesktopPlatformLinux.cpp
@@ -6,7 +6,6 @@
 #include "Misc/Guid.h"
 #include "Misc/ConfigCacheIni.h"
 #include "Misc/FeedbackContext.h"
-#include "Misc/SecureHash.h"
 #include "DesktopPlatformPrivate.h"
 #include "Modules/ModuleManager.h"
 #include "Linux/LinuxApplication.h"
@@ -449,59 +448,6 @@ bool FDesktopPlatformLinux::RunUnrealBuildTool(const FText& Description, const F
 	return FFeedbackContextMarkup::PipeProcessOutput(Description, UnrealBuildToolPath, Arguments, Warn, &OutExitCode) && OutExitCode == 0;
 }
 
-bool FDesktopPlatformLinux::IsUnrealBuildToolRunning()
-{
-	//TODO 5.1, make this a common routine instead of duplicated
-	FString RunsDir = FPaths::Combine(FPaths::EngineIntermediateDir(), TEXT("UbtRuns"));
-	if (!FPaths::DirectoryExists(RunsDir))
-	{
-		return false;
-	}
-
-	bool bIsRunning = false;
-	IFileManager::Get().IterateDirectory(*RunsDir, [&bIsRunning](const TCHAR* Pathname, bool bIsDirectory)
-		{
-			if (!bIsDirectory)
-			{
-				bool bDeleteFile = true;
-
-				FString Filename = FPaths::GetBaseFilename(FString(Pathname));
-				const TCHAR* Delim = FCString::Strchr(*Filename, '_');
-				if (Delim != nullptr)
-				{
-					FString Pid(*Filename, UE_PTRDIFF_TO_INT32(Delim - *Filename));
-					int ProcessId = 0;
-					LexFromString(ProcessId, *Pid);
-					const FString EntryFullPath = FPlatformProcess::GetApplicationName(ProcessId);
-					if (!EntryFullPath.IsEmpty())
-					{
-						FString EntryFullPathUpper = EntryFullPath.ToUpper();
-						const FTCHARToUTF8 Utf8String(*EntryFullPathUpper);
-						FMD5Hash Hash;
-						LexFromString(Hash, Delim + 1);
-
-						FMD5 Md5Gen;
-						Md5Gen.Update(reinterpret_cast<const uint8*>(Utf8String.Get()), Utf8String.Length());
-						FMD5Hash TestHash;
-						TestHash.Set(Md5Gen);
-						if (Hash == TestHash)
-						{
-							bDeleteFile = false;
-							bIsRunning = true;
-						}
-					}
-					if (bDeleteFile)
-					{
-						IFileManager::Get().Delete(Pathname);
-					}
-				}
-			}
-			return true;
-		});
-
-	return bIsRunning;
-}
-
 FFeedbackContext* FDesktopPlatformLinux::GetNativeFeedbackContext()
 {
 	//unimplemented();
diff --git a/Engine/Source/Developer/DesktopPlatform/Private/Linux/DesktopPlatformLinux.h b/Engine/Source/Developer/DesktopPlatform/Private/Linux/DesktopPlatformLinux.h
--- a/Engine/Source/Developer/DesktopPlatform/Private/Linux/DesktopPlatformLinux.h
+++ b/Engine/Source/Developer/DesktopPlatform/Private/Linux/DesktopPlatformLinux.h
@@ -27,7 +27,6 @@ public:
 
 	using FDesktopPlatformBase::RunUnrealBuildTool;
 	virtual bool RunUnrealBuildTool(const FText& Description, const FString& RootDir, const FString& Arguments, FFeedbackContext* Warn, int32& OutExitCode) override;
-	virtual bool IsUnrealBuildToolRunning() override;
 
 	virtual FFeedbackContext* GetNativeFeedbackContext() override;
 
diff --git a/Engine/Source/Developer/DesktopPlatform/Private/Mac/DesktopPlatformMac.cpp b/Engine/Source/Developer/DesktopPlatform/Private/Mac/DesktopPlatformMac.cpp
--- a/Engine/Source/Developer/DesktopPlatform/Private/Mac/DesktopPlatformMac.cpp
+++ b/Engine/Source/Developer/DesktopPlatform/Private/Mac/DesktopPlatformMac.cpp
@@ -8,7 +8,6 @@
 #include "Misc/Paths.h"
 #include "Misc/ConfigCacheIni.h"
 #include "Misc/Guid.h"
-#include "Misc/SecureHash.h"
 #include "HAL/FileManager.h"
 
 #define LOCTEXT_NAMESPACE "DesktopPlatform"
@@ -626,59 +625,6 @@ bool FDesktopPlatformMac::RunUnrealBuildTool(const FText& Description, const FSt
 	return FFeedbackContextMarkup::PipeProcessOutput(Description, UnrealBuildToolPath, Arguments, Warn, &OutExitCode) && OutExitCode == 0;
 }
 
-bool FDesktopPlatformMac::IsUnrealBuildToolRunning()
-{
-	//TODO 5.1, make this a common routine instead of duplicated
-	FString RunsDir = FPaths::Combine(FPaths::EngineIntermediateDir(), TEXT("UbtRuns"));
-	if (!FPaths::DirectoryExists(RunsDir))
-	{
-		return false;
-	}
-
-	bool bIsRunning = false;
-	IFileManager::Get().IterateDirectory(*RunsDir, [&bIsRunning](const TCHAR* Pathname, bool bIsDirectory)
-		{
-			if (!bIsDirectory)
-			{
-				bool bDeleteFile = true;
-
-				FString Filename = FPaths::GetBaseFilename(FString(Pathname));
-				const TCHAR* Delim = FCString::Strchr(*Filename, '_');
-				if (Delim != nullptr)
-				{
-					FString Pid(*Filename, UE_PTRDIFF_TO_INT32(Delim - *Filename));
-					int ProcessId = 0;
-					LexFromString(ProcessId, *Pid);
-					const FString EntryFullPath = FPlatformProcess::GetApplicationName(ProcessId);
-					if (!EntryFullPath.IsEmpty())
-					{
-						FString EntryFullPathUpper = EntryFullPath.ToUpper();
-						const FTCHARToUTF8 Utf8String(*EntryFullPathUpper);
-						FMD5Hash Hash;
-						LexFromString(Hash, Delim + 1);
-
-						FMD5 Md5Gen;
-						Md5Gen.Update(reinterpret_cast<const uint8*>(Utf8String.Get()), Utf8String.Length());
-						FMD5Hash TestHash;
-						TestHash.Set(Md5Gen);
-						if (Hash == TestHash)
-						{
-							bDeleteFile = false;
-							bIsRunning = true;
-						}
-					}
-					if (bDeleteFile)
-					{
-						IFileManager::Get().Delete(Pathname);
-					}
-				}
-			}
-			return true;
-		});
-
-	return bIsRunning;
-}
-
 FFeedbackContext* FDesktopPlatformMac::GetNativeFeedbackContext()
 {
 	static FMacNativeFeedbackContext Warn;
diff --git a/Engine/Source/Developer/DesktopPlatform/Private/Mac/DesktopPlatformMac.h b/Engine/Source/Developer/DesktopPlatform/Private/Mac/DesktopPlatformMac.h
--- a/Engine/Source/Developer/DesktopPlatform/Private/Mac/DesktopPlatformMac.h
+++ b/Engine/Source/Developer/DesktopPlatform/Private/Mac/DesktopPlatformMac.h
@@ -22,7 +22,6 @@ public:
 
 	using FDesktopPlatformBase::RunUnrealBuildTool;
 	virtual bool RunUnrealBuildTool(const FText& Description, const FString& RootDir, const FString& Arguments, FFeedbackContext* Warn, int32& OutExitCode) override;
-	virtual bool IsUnrealBuildToolRunning() override;
 
 	virtual FFeedbackContext* GetNativeFeedbackContext() override;
 
diff --git a/Engine/Source/Developer/DesktopPlatform/Private/Windows/DesktopPlatformWindows.cpp b/Engine/Source/Developer/DesktopPlatform/Private/Windows/DesktopPlatformWindows.cpp
--- a/Engine/Source/Developer/DesktopPlatform/Private/Windows/DesktopPlatformWindows.cpp
+++ b/Engine/Source/Developer/DesktopPlatform/Private/Windows/DesktopPlatformWindows.cpp
@@ -7,7 +7,6 @@
 #include "Microsoft/COMPointer.h"
 #include "Misc/Paths.h"
 #include "Misc/Guid.h"
-#include "Misc/SecureHash.h"
 #include "HAL/FileManager.h"
 
 #include "Windows/AllowWindowsPlatformTypes.h"
@@ -478,59 +477,6 @@ bool FDesktopPlatformWindows::RunUnrealBuildTool(const FText& Description, const
 	return FFeedbackContextMarkup::PipeProcessOutput(Description, UnrealBuildToolPath, Arguments, Warn, &OutExitCode) && OutExitCode == 0;
 }
 
-bool FDesktopPlatformWindows::IsUnrealBuildToolRunning()
-{
-	//TODO 5.1, make this a common routine instead of duplicated
-	FString RunsDir = FPaths::Combine(FPaths::EngineIntermediateDir(), TEXT("UbtRuns"));
-	if (!FPaths::DirectoryExists(RunsDir))
-	{
-		return false;
-	}
-
-	bool bIsRunning = false;
-	IFileManager::Get().IterateDirectory(*RunsDir, [&bIsRunning](const TCHAR* Pathname, bool bIsDirectory)
-		{
-			if (!bIsDirectory)
-			{
-				bool bDeleteFile = true;
-
-				FString Filename = FPaths::GetBaseFilename(FString(Pathname));
-				const TCHAR* Delim = FCString::Strchr(*Filename, '_');
-				if (Delim != nullptr)
-				{
-					FString Pid(*Filename, UE_PTRDIFF_TO_INT32(Delim - *Filename));
-					int ProcessId = 0;
-					LexFromString(ProcessId, *Pid);
-					const FString EntryFullPath = FPlatformProcess::GetApplicationName(ProcessId);
-					if (!EntryFullPath.IsEmpty())
-					{
-						FString EntryFullPathUpper = EntryFullPath.ToUpper();
-						const FTCHARToUTF8 Utf8String(*EntryFullPathUpper);
-						FMD5Hash Hash;
-						LexFromString(Hash, Delim + 1);
-
-						FMD5 Md5Gen;
-						Md5Gen.Update(reinterpret_cast<const uint8*>(Utf8String.Get()), Utf8String.Length());
-						FMD5Hash TestHash;
-						TestHash.Set(Md5Gen);
-						if (Hash == TestHash)
-						{
-							bDeleteFile = false;
-							bIsRunning = true;
-						}
-					}
-					if (bDeleteFile)
-					{
-						IFileManager::Get().Delete(Pathname);
-					}
-				}
-			}
-			return true;
-		});
-
-	return bIsRunning;
-}
-
 FFeedbackContext* FDesktopPlatformWindows::GetNativeFeedbackContext()
 {
 	static FWindowsNativeFeedbackContext FeedbackContext;
diff --git a/Engine/Source/Developer/DesktopPlatform/Private/Windows/DesktopPlatformWindows.h b/Engine/Source/Developer/DesktopPlatform/Private/Windows/DesktopPlatformWindows.h
--- a/Engine/Source/Developer/DesktopPlatform/Private/Windows/DesktopPlatformWindows.h
+++ b/Engine/Source/Developer/DesktopPlatform/Private/Windows/DesktopPlatformWindows.h
@@ -30,7 +30,6 @@ public:
 
 	using FDesktopPlatformBase::RunUnrealBuildTool;
 	virtual bool RunUnrealBuildTool(const FText& Description, const FString& RootDir, const FString& Arguments, FFeedbackContext* Warn, int32& OutExitCode) override;
-	virtual bool IsUnrealBuildToolRunning() override;
 
 	virtual FFeedbackContext* GetNativeFeedbackContext() override;
 
