OidcToken - Fixed issue running interactive login on mac and linux. Also fixed issue were asking to wait for login caused a crash as well as tweaking the error message when a user decided to not finish the interactive login.

diff --git a/Engine/Source/Developer/DesktopPlatform/Private/DesktopPlatformBase.cpp b/Engine/Source/Developer/DesktopPlatform/Private/DesktopPlatformBase.cpp
--- a/Engine/Source/Developer/DesktopPlatform/Private/DesktopPlatformBase.cpp
+++ b/Engine/Source/Developer/DesktopPlatform/Private/DesktopPlatformBase.cpp
@@ -829,7 +829,7 @@ bool FDesktopPlatformBase::GetOidcAccessToken(const FString& RootDir, const FStr
 			bRes = GetOidcAccessTokenInteractive(RootDir, Arguments, Warn, ExitCode);
 			if (!bRes)
 			{
-				UE_LOG(LogDesktopPlatform, Error, TEXT("Unable to allocate an access token. Interactive login failed, make sure you are logged in using UGS or using the UGS cli command 'login'. Provider used: '%s'. Ran OidcToken (project file is '%s', exe path is '%s')"), *ProviderIdentifier, *ProjectFileName, *GetOidcTokenExecutableFilename(RootDir));
+				UE_LOG(LogDesktopPlatform, Error, TEXT("Unable to allocate an access token. Interactive login failed, make sure you are assigned access and are able to login in the created browser window. Provider used: '%s'. Ran OidcToken (project file is '%s', exe path is '%s')"), *ProviderIdentifier, *ProjectFileName, *GetOidcTokenExecutableFilename(RootDir));
 				return false;
 			}
 		}
@@ -980,15 +980,18 @@ bool FDesktopPlatformBase::GetOidcAccessTokenInteractive(const FString& RootDir,
 		FPlatformProcess::Sleep(0.1f);
 	}
 
-	bool bGotReturnCode = FPlatformProcess::GetProcReturnCode(ProcHandle, &OutReturnCode);		
-	check(bGotReturnCode);
-	FPlatformProcess::ClosePipe(PipeRead, PipeWrite);
-	
 	if (!bIsFinished)
 	{
 		OutReturnCode = -1;
+		FPlatformProcess::ClosePipe(PipeRead, PipeWrite);
+		FPlatformProcess::TerminateProc(ProcHandle);
+		return false;
 	}
 
+	bool bGotReturnCode = FPlatformProcess::GetProcReturnCode(ProcHandle, &OutReturnCode);		
+	check(bGotReturnCode);
+	FPlatformProcess::ClosePipe(PipeRead, PipeWrite);
+
 	return bInvoked;
 }
 
diff --git a/Engine/Source/Developer/DesktopPlatform/Private/Linux/DesktopPlatformLinux.cpp b/Engine/Source/Developer/DesktopPlatform/Private/Linux/DesktopPlatformLinux.cpp
--- a/Engine/Source/Developer/DesktopPlatform/Private/Linux/DesktopPlatformLinux.cpp
+++ b/Engine/Source/Developer/DesktopPlatform/Private/Linux/DesktopPlatformLinux.cpp
@@ -462,7 +462,7 @@ FString FDesktopPlatformLinux::GetUserTempPath()
 
 FString FDesktopPlatformLinux::GetOidcTokenExecutableFilename(const FString& RootDir) const
 {	
-	return FPaths::ConvertRelativePathToFull(RootDir / TEXT("Engine/Binaries/DotNET/OidcToken/linux-x64/OidcToken.exe"));
+	return FPaths::ConvertRelativePathToFull(RootDir / TEXT("Engine/Binaries/DotNET/OidcToken/linux-x64/OidcToken"));
 }
 
 #undef LOCTEXT_NAMESPACE
diff --git a/Engine/Source/Developer/DesktopPlatform/Private/Mac/DesktopPlatformMac.cpp b/Engine/Source/Developer/DesktopPlatform/Private/Mac/DesktopPlatformMac.cpp
--- a/Engine/Source/Developer/DesktopPlatform/Private/Mac/DesktopPlatformMac.cpp
+++ b/Engine/Source/Developer/DesktopPlatform/Private/Mac/DesktopPlatformMac.cpp
@@ -638,6 +638,6 @@ FString FDesktopPlatformMac::GetUserTempPath()
 
 FString FDesktopPlatformMac::GetOidcTokenExecutableFilename(const FString& RootDir) const
 {
-	return FPaths::ConvertRelativePathToFull(RootDir / TEXT("Engine/Binaries/DotNET/OidcToken/mac-x64/OidcToken.exe"));
+	return FPaths::ConvertRelativePathToFull(RootDir / TEXT("Engine/Binaries/DotNET/OidcToken/osx-x64/OidcToken"));
 }
 #undef LOCTEXT_NAMESPACE
