OidcToken - Fixed issues specific to handle tokens on Mac, which handles cases where the token does not exist better.
Also improved logging when a unknown issue happens.
Lastly the error code returned from the mac process seems to be converted to a unsigned 32 bit so checking for that value as well when checking the exit code that means we should upgrade the process to a interactive login.

diff --git a/Engine/Source/Developer/DesktopPlatform/Private/DesktopPlatformBase.cpp b/Engine/Source/Developer/DesktopPlatform/Private/DesktopPlatformBase.cpp
--- a/Engine/Source/Developer/DesktopPlatform/Private/DesktopPlatformBase.cpp
+++ b/Engine/Source/Developer/DesktopPlatform/Private/DesktopPlatformBase.cpp
@@ -822,7 +822,7 @@ bool FDesktopPlatformBase::GetOidcAccessToken(const FString& RootDir, const FStr
 	FString ProcessStdout;
 	bRes = InvokeOidcTokenToolSync(LOCTEXT("GetOidcAccessToken", "Fetching OIDC Access Token..."), RootDir, UnattendedArguments, Warn, ExitCode, ProcessStdout);
 
-	if (ExitCode == -1337)
+	if (ExitCode == -1337 || ExitCode == 16775879)
 	{
 		if (!Unattended)
 		{
@@ -868,6 +868,8 @@ bool FDesktopPlatformBase::GetOidcAccessToken(const FString& RootDir, const FStr
 		}
 	}
 
+	UE_LOG(LogDesktopPlatform, Warning, TEXT("Failed to run OidcToken (project file is '%s', exe path is '%s'). No result file found at '%s', closed with exit code: %d"), *ProjectFileName, *GetOidcTokenExecutableFilename(RootDir), *ResultFilePath, ExitCode);
+
 	return false;
 }
 
diff --git a/Engine/Source/Programs/Shared/EpicGames.OIDC/MacOSTokenStore.cs b/Engine/Source/Programs/Shared/EpicGames.OIDC/MacOSTokenStore.cs
--- a/Engine/Source/Programs/Shared/EpicGames.OIDC/MacOSTokenStore.cs
+++ b/Engine/Source/Programs/Shared/EpicGames.OIDC/MacOSTokenStore.cs
@@ -75,9 +75,14 @@ namespace EpicGames.OIDC
 				return null; // not found
 			}
 
+			if (p.ExitCode == 36)
+			{
+				return null; // not found
+			}
+
 			if (p.ExitCode != 0)
 			{
-				throw new Exception("Unhandled exitcode {p.ExitCode} when interacting with keychain");
+				throw new Exception($"Unhandled exitcode {p.ExitCode} when interacting with keychain");
 			}
 
 			return sb.ToString();
