OidcToken - Allow users to abort the oidc login in the prompt that warns we are about to spawn a browser

diff --git a/Engine/Source/Developer/DesktopPlatform/Private/DesktopPlatformBase.cpp b/Engine/Source/Developer/DesktopPlatform/Private/DesktopPlatformBase.cpp
--- a/Engine/Source/Developer/DesktopPlatform/Private/DesktopPlatformBase.cpp
+++ b/Engine/Source/Developer/DesktopPlatform/Private/DesktopPlatformBase.cpp
@@ -920,7 +920,14 @@ bool FDesktopPlatformBase::GetOidcAccessTokenInteractive(const FString& RootDir,
 {
 	FText OidcInteractivePromptTitle = NSLOCTEXT("OidcToken", "OidcToken_InteractiveLaunchPromptTitle", "Unreal Engine - User Login");
 	FText OidcInteractiveLaunchPromptText = NSLOCTEXT("OidcToken", "OidcToken_InteractiveLaunch", "Authentication is required for Http Derived Data Cache to work. Your system browser will now be opened for you to finish the login process.");
-	FPlatformMisc::MessageBoxExt(EAppMsgType::Ok, *OidcInteractiveLaunchPromptText.ToString(), *OidcInteractivePromptTitle.ToString());
+	EAppReturnType::Type userAcknowledgedResult = FPlatformMisc::MessageBoxExt(EAppMsgType::OkCancel, *OidcInteractiveLaunchPromptText.ToString(), *OidcInteractivePromptTitle.ToString());
+
+	if (userAcknowledgedResult != EAppReturnType::Ok)
+	{
+		OutReturnCode = -1;
+		return false;
+	}
+
 	// user has acknowledged the login, we update the editor progress and then we run oidc token to spawn the browser and prompt the login
 
 	FScopedSlowTask WaitForOidcTokenSlowTask(0, NSLOCTEXT("OidcToken", "OidcToken_WaitingForToken", "Waiting for OidcToken to finish login"));
