OidcToken - Fixed issue on linux were exit code was not correctly detected and thus the initial verification if a token exists fails.

diff --git a/Engine/Source/Developer/DesktopPlatform/Private/DesktopPlatformBase.cpp b/Engine/Source/Developer/DesktopPlatform/Private/DesktopPlatformBase.cpp
--- a/Engine/Source/Developer/DesktopPlatform/Private/DesktopPlatformBase.cpp
+++ b/Engine/Source/Developer/DesktopPlatform/Private/DesktopPlatformBase.cpp
@@ -822,7 +822,8 @@ bool FDesktopPlatformBase::GetOidcAccessToken(const FString& RootDir, const FStr
 	FString ProcessStdout;
 	bRes = InvokeOidcTokenToolSync(LOCTEXT("GetOidcAccessToken", "Fetching OIDC Access Token..."), RootDir, UnattendedArguments, Warn, ExitCode, ProcessStdout);
 
-	if (ExitCode == -1337 || ExitCode == 16775879)
+	// the tool always returns -1337 as the exit code but osx and linux interprets this exit code differently
+	if (ExitCode == -1337 || ExitCode == 16775879 || ExitCode == 199)
 	{
 		if (!Unattended)
 		{
