OidcToken - Changed exit code when determining if a token exists to a positive small number (10) which will avoid inconsistencies between desktop platform in how to handle a negative number.

diff --git a/Engine/Source/Developer/DesktopPlatform/Private/DesktopPlatformBase.cpp b/Engine/Source/Developer/DesktopPlatform/Private/DesktopPlatformBase.cpp
--- a/Engine/Source/Developer/DesktopPlatform/Private/DesktopPlatformBase.cpp
+++ b/Engine/Source/Developer/DesktopPlatform/Private/DesktopPlatformBase.cpp
@@ -822,8 +822,7 @@ bool FDesktopPlatformBase::GetOidcAccessToken(const FString& RootDir, const FStr
 	FString ProcessStdout;
 	bRes = InvokeOidcTokenToolSync(LOCTEXT("GetOidcAccessToken", "Fetching OIDC Access Token..."), RootDir, UnattendedArguments, Warn, ExitCode, ProcessStdout);
 
-	// the tool always returns -1337 as the exit code but osx and linux interprets this exit code differently
-	if (ExitCode == -1337 || ExitCode == 16775879 || ExitCode == 199)
+	if (ExitCode == 10)
 	{
 		if (!Unattended)
 		{
diff --git a/Engine/Source/Programs/OidcToken/TokenService.cs b/Engine/Source/Programs/OidcToken/TokenService.cs
--- a/Engine/Source/Programs/OidcToken/TokenService.cs
+++ b/Engine/Source/Programs/OidcToken/TokenService.cs
@@ -44,12 +44,12 @@ namespace OidcToken
 				catch (UnableToAllocateTokenException)
 				{
 					Logger.LogWarning("Was unable to allocate a token");
-					ExitCode = -1337;
+					ExitCode = 10;
 				}
 				catch (Exception ex)
 				{
 					Logger.LogError(ex, "Unhandled exception!");
-					ExitCode = -1;
+					ExitCode = 1;
 				}
 				finally
 				{
