Fixed a few issues where IsTargetInfoValid would return that all targets weren't found. This would result in a call to UBT to get the TargetInfo and sometimes took up to 22+secs to call.

diff --git a/Engine/Source/Developer/DesktopPlatform/Private/DesktopPlatformBase.cpp b/Engine/Source/Developer/DesktopPlatform/Private/DesktopPlatformBase.cpp
--- a/Engine/Source/Developer/DesktopPlatform/Private/DesktopPlatformBase.cpp
+++ b/Engine/Source/Developer/DesktopPlatform/Private/DesktopPlatformBase.cpp
@@ -1067,15 +1067,17 @@ FProcHandle FDesktopPlatformBase::InvokeOidcTokenToolAsync(const FString& InArgu
 
 struct FTargetFileVisitor : IPlatformFile::FDirectoryStatVisitor
 {
+	const TSet<FString>& OriginalTargetNames;
 	TSet<FString>& RemainingTargetNames;
 	FDateTime MaxDateTime;
 	TArray<FString> SubDirectories;
 	bool bSearchSubDirectories;
 
-	FTargetFileVisitor(TSet<FString>& InRemainingTargetNames, FDateTime InMaxDateTime)
-		: RemainingTargetNames(InRemainingTargetNames)
+	FTargetFileVisitor(const TSet<FString>& InOriginalTargetNames, TSet<FString>& InRemainingTargetNames, FDateTime InMaxDateTime)
+		: OriginalTargetNames(InOriginalTargetNames)
+		, RemainingTargetNames(InRemainingTargetNames)
 		, MaxDateTime(InMaxDateTime)
-		, bSearchSubDirectories(false)
+		, bSearchSubDirectories(true)
 	{
 	}
 
@@ -1087,19 +1089,49 @@ struct FTargetFileVisitor : IPlatformFile::FDirectoryStatVisitor
 			return true;
 		}
 
-		int32 Length = FCString::Strlen(FileNameOrDirectory);
-
+		// NOTE: This code needs to behave the same as FindAllRulesSourceFiles in Rules.cs
 		static const TCHAR TargetExt[] = TEXT(".Target.cs");
 		static const int32 TargetExtLen = UE_ARRAY_COUNT(TargetExt) - 1;
+		static const TCHAR ModuleExt[] = TEXT(".Build.cs");
+		static const int32 ModuleExtLen = UE_ARRAY_COUNT(ModuleExt) - 1;
+		static const TCHAR AutomationCsprojExt[] = TEXT(".automation.csproj");
+		static const int32 AutomationCsprojExtLen = UE_ARRAY_COUNT(AutomationCsprojExt) - 1;
+		static const TCHAR UBTCsprojExt[] = TEXT(".ubtplugin.csproj");
+		static const int32 UBTCsprojExtLen = UE_ARRAY_COUNT(UBTCsprojExt) - 1;
+		static const TCHAR UBTIgnoreExt[] = TEXT(".ubtignore");
+		static const int32 UBTIgnoreExtLen = UE_ARRAY_COUNT(UBTIgnoreExt) - 1;
+
+		int32 Length = FCString::Strlen(FileNameOrDirectory);
 		if (Length > TargetExtLen && FCString::Stricmp(FileNameOrDirectory + Length - TargetExtLen, TargetExt) == 0)
 		{
 			FString TargetName = FPaths::GetCleanFilename(FString(Length - TargetExtLen, FileNameOrDirectory));
+
+			// skip target rules that are platform extension or platform group specializations
+			// Matches logic found in QueryTargetsMode.cs WriteTargetInfo
+			FString Start, End;
+			if (TargetName.Split(TEXT("_"), &Start, &End) && OriginalTargetNames.Contains(Start))
+			{
+				return true;
+			}
+
 			return (StatData.ModificationTime < MaxDateTime && RemainingTargetNames.Remove(TargetName) == 1);
 		}
-		
-		static const TCHAR ModuleExt[] = TEXT(".Build.cs");
-		static const int32 ModuleExtLen = UE_ARRAY_COUNT(ModuleExt) - 1;
-		if (Length > ModuleExtLen && FCString::Stricmp(FileNameOrDirectory + Length - ModuleExtLen, ModuleExt) == 0)
+		else if (Length > ModuleExtLen && FCString::Stricmp(FileNameOrDirectory + Length - ModuleExtLen, ModuleExt) == 0)
+		{
+			bSearchSubDirectories = false;
+			return true;
+		}
+		else if (Length > AutomationCsprojExtLen && FCString::Stricmp(FileNameOrDirectory + Length - AutomationCsprojExtLen, AutomationCsprojExt) == 0)
+		{
+			bSearchSubDirectories = false;
+			return true;
+		}
+		else if (Length > UBTCsprojExtLen && FCString::Stricmp(FileNameOrDirectory + Length - UBTCsprojExtLen, UBTCsprojExt) == 0)
+		{
+			bSearchSubDirectories = false;
+			return true;
+		}
+		else if (Length > UBTIgnoreExtLen && FCString::Stricmp(FileNameOrDirectory + Length - UBTIgnoreExtLen, UBTIgnoreExt) == 0)
 		{
 			bSearchSubDirectories = false;
 			return true;
@@ -1109,7 +1141,7 @@ struct FTargetFileVisitor : IPlatformFile::FDirectoryStatVisitor
 	}
 };
 
-bool IsTargetInfoValid(const TArray<FTargetInfo>& Targets, const FString& SourceDir, const FDateTime& LastModifiedTime)
+bool IsTargetInfoValid(const TArray<FTargetInfo>& Targets, TArray<FString>& DirectoryNames, const FDateTime& LastModifiedTime)
 {
 	if (FApp::GetEngineIsPromotedBuild())
 	{
@@ -1124,11 +1156,12 @@ bool IsTargetInfoValid(const TArray<FTargetInfo>& Targets, const FString& Source
 		RemainingTargetNames.Add(Target.Name);
 	}
 
+	TSet<FString> OriginalTargetNames = RemainingTargetNames;
+
 	// Loop through all the directories
-	TArray<FString> DirectoryNames = { SourceDir };
 	for(int Idx = 0; Idx < DirectoryNames.Num(); Idx++)
 	{
-		FTargetFileVisitor Visitor(RemainingTargetNames, LastModifiedTime);
+		FTargetFileVisitor Visitor(OriginalTargetNames, RemainingTargetNames, LastModifiedTime);
 		if(!IFileManager::Get().IterateDirectoryStat(*DirectoryNames[Idx], Visitor))
 		{
 			return false;
@@ -1178,7 +1211,8 @@ const TArray<FTargetInfo>& FDesktopPlatformBase::GetTargetsForProject(const FStr
 	{
 		// Read it in and check it's still valid
 		TArray<FTargetInfo> NewTargets;
-		if(ReadTargetInfo(InfoFileName, NewTargets) && IsTargetInfoValid(NewTargets, ProjectSourceDir, StatData.ModificationTime))
+		TArray<FString> DirectoryNames = { ProjectSourceDir, ProjectDir / TEXT("Platforms"), ProjectDir / TEXT("Restricted") };
+		if(ReadTargetInfo(InfoFileName, NewTargets) && IsTargetInfoValid(NewTargets, DirectoryNames, StatData.ModificationTime))
 		{
 			return ProjectFileToTargets.Emplace(MoveTemp(NormalizedProjectFile), MoveTemp(NewTargets));
 		}
