OidcToken - If users answer no to waiting for interactive login, but the login has actually completed we now continue with the valid token we have.

diff --git a/Engine/Source/Developer/DesktopPlatform/Private/DesktopPlatformBase.cpp b/Engine/Source/Developer/DesktopPlatform/Private/DesktopPlatformBase.cpp
--- a/Engine/Source/Developer/DesktopPlatform/Private/DesktopPlatformBase.cpp
+++ b/Engine/Source/Developer/DesktopPlatform/Private/DesktopPlatformBase.cpp
@@ -981,6 +981,7 @@ bool FDesktopPlatformBase::GetOidcAccessTokenInteractive(const FString& RootDir,
 			FText OidcLongWaitPromptText = NSLOCTEXT("OidcToken", "OidcToken_LongWaitPromptText", "Login is taking a long time, make sure you have entered your credentials in your browser window. It can be in a tab in an already existing window. Keep waiting?");
 			if (FPlatformMisc::MessageBoxExt(EAppMsgType::YesNo, *OidcLongWaitPromptText.ToString(), *OidcLongWaitPromptTitle.ToString()) == EAppReturnType::No)
 			{
+				bIsFinished = !FPlatformProcess::IsProcRunning(ProcHandle);
 				break;
 			}
 			// change phase so we do not prompt the user again
@@ -989,6 +990,7 @@ bool FDesktopPlatformBase::GetOidcAccessTokenInteractive(const FString& RootDir,
 
 		if (WaitForOidcTokenSlowTask.ShouldCancel())
 		{
+			bIsFinished = !FPlatformProcess::IsProcRunning(ProcHandle);
 			break;
 		}
 		FPlatformProcess::Sleep(0.1f);
