PR #9906: Fix Linux duplicate engine association entries in Install.ini. (Contributed by [REDACTED_PERSON])

diff --git a/Engine/Source/Developer/DesktopPlatform/Private/DesktopPlatformBase.cpp b/Engine/Source/Developer/DesktopPlatform/Private/DesktopPlatformBase.cpp
--- a/Engine/Source/Developer/DesktopPlatform/Private/DesktopPlatformBase.cpp
+++ b/Engine/Source/Developer/DesktopPlatform/Private/DesktopPlatformBase.cpp
@@ -153,7 +153,7 @@ bool FDesktopPlatformBase::GetEngineIdentifierFromRootDir(const FString &RootDir
 	}
 
 	// Otherwise just try to add it
-	return RegisterEngineInstallation(RootDir, OutIdentifier);
+	return RegisterEngineInstallation(NormalizedRootDir, OutIdentifier);
 }
 
 bool FDesktopPlatformBase::GetDefaultEngineIdentifier(FString &OutId)
diff --git a/Engine/Source/Developer/DesktopPlatform/Private/Linux/DesktopPlatformLinux.cpp b/Engine/Source/Developer/DesktopPlatform/Private/Linux/DesktopPlatformLinux.cpp
--- a/Engine/Source/Developer/DesktopPlatform/Private/Linux/DesktopPlatformLinux.cpp
+++ b/Engine/Source/Developer/DesktopPlatform/Private/Linux/DesktopPlatformLinux.cpp
@@ -135,7 +135,7 @@ bool FDesktopPlatformLinux::RegisterEngineInstallation(const FString &RootDir, F
 		ConfigFile.Read(ConfigPath);
 
 		FConfigSection &Section = ConfigFile.FindOrAdd(TEXT("Installations"));
-		OutIdentifier = FGuid::NewGuid().ToString(EGuidFormats::DigitsWithHyphensInBraces);
+		OutIdentifier = FGuid::NewGuid().ToString(EGuidFormats::DigitsWithHyphens);
 		Section.AddUnique(*OutIdentifier, RootDir);
 
 		ConfigFile.Dirty = true;
@@ -205,13 +205,13 @@ void FDesktopPlatformLinux::EnumerateEngineInstallations(TMap<FString, FString>
 		{
 			FGuid IdGuid;
 			FGuid::Parse(Key->ToString(), IdGuid);
-			EngineId = IdGuid.ToString(EGuidFormats::DigitsWithHyphensInBraces);
+			EngineId = IdGuid.ToString(EGuidFormats::DigitsWithHyphens);
 		}
 		else
 		{
 			if (!OutInstallations.FindKey(EngineDir))
 			{
-				EngineId = FGuid::NewGuid().ToString(EGuidFormats::DigitsWithHyphensInBraces);
+				EngineId = FGuid::NewGuid().ToString(EGuidFormats::DigitsWithHyphens);
 				SectionsToAdd.AddUnique(*EngineId, EngineDir);
 
 				ConfigFile.Dirty = true;
