Wait up to 30 seconds before prompting users to finish the login - it can take a while if users need to enter their password and such, it is still quite short but on the other hand we do not want the editor to just sit around waiting for ever in case something has gone wrong.

diff --git a/Engine/Source/Developer/DesktopPlatform/Private/DesktopPlatformBase.cpp b/Engine/Source/Developer/DesktopPlatform/Private/DesktopPlatformBase.cpp
--- a/Engine/Source/Developer/DesktopPlatform/Private/DesktopPlatformBase.cpp
+++ b/Engine/Source/Developer/DesktopPlatform/Private/DesktopPlatformBase.cpp
@@ -974,8 +974,8 @@ bool FDesktopPlatformBase::GetOidcAccessTokenInteractive(const FString& RootDir,
 			UE_LOG(LogDesktopPlatform, Display, TEXT("Waiting for OidcToken to finish login..."));
 			DurationPhase = EWaitDurationPhase::Prompt;
 		}
-		// once we have waited for 10 seconds without success we give the user a option to abort
-		else if (WaitDuration > 10.0 && DurationPhase == EWaitDurationPhase::Prompt)
+		// once we have waited for 30 seconds without success we give the user a option to abort
+		else if (WaitDuration > 30.0 && DurationPhase == EWaitDurationPhase::Prompt)
 		{
 			FText OidcLongWaitPromptTitle = NSLOCTEXT("OidcToken", "OidcToken_LongWaitPromptTitle", "Wait for user login?");
 			FText OidcLongWaitPromptText = NSLOCTEXT("OidcToken", "OidcToken_LongWaitPromptText", "Login is taking a long time, make sure you have entered your credentials in your browser window. It can be in a tab in an already existing window. Keep waiting?");
