Logging: Updated Mac feedback contexts to support structured logging

This relies on the base feedback context doing most of the work, as on other platforms.

Coloring was removed from the default Mac feedback context because it has no effect on the default output devices anyway. Output devices that do want to set color based on verbosity should incorporate that logic themselves in a thread-safe way instead of relying on the color being changed globally.

#rnx

diff --git a/Engine/Source/Developer/DesktopPlatform/Private/Mac/MacNativeFeedbackContext.cpp b/Engine/Source/Developer/DesktopPlatform/Private/Mac/MacNativeFeedbackContext.cpp
--- a/Engine/Source/Developer/DesktopPlatform/Private/Mac/MacNativeFeedbackContext.cpp
+++ b/Engine/Source/Developer/DesktopPlatform/Private/Mac/MacNativeFeedbackContext.cpp
@@ -1,6 +1,7 @@
 // Copyright Epic Games, Inc. All Rights Reserved.
 
 #include "MacNativeFeedbackContext.h"
+#include "Logging/StructuredLog.h"
 #include "Mac/MacApplication.h"
 #include "Mac/CocoaThread.h"
 #include "Misc/OutputDeviceHelper.h"
@@ -282,17 +283,32 @@ FMacNativeFeedbackContext::~FMacNativeFeedbackContext()
 	}, UnrealNilEventMode, true);
 }
 
-void FMacNativeFeedbackContext::Serialize( const TCHAR* Data, ELogVerbosity::Type Verbosity, const class FName& Category )
+void FMacNativeFeedbackContext::Serialize(const TCHAR* Data, ELogVerbosity::Type Verbosity, const FName& Category)
 {
-	if( !GLog->IsRedirectingTo( this ) )
-	{
-		GLog->Serialize( Data, Verbosity, Category );
-	}
-	
-	if(WindowController != NULL && bShowingConsoleForSlowTask)
+	Serialize(Data, Verbosity, Category, -1.0);
+}
+
+void FMacNativeFeedbackContext::Serialize(const TCHAR* Data, ELogVerbosity::Type Verbosity, const FName& Category, double Time)
+{
+	FFeedbackContext::Serialize(Data, Verbosity, Category, Time);
+	SerializeToWindow(Data, Verbosity, Category, Time);
+}
+
+void FMacNativeFeedbackContext::SerializeRecord(const UE::FLogRecord& Record)
+{
+	FFeedbackContext::SerializeRecord(Record);
+
+	TStringBuilder<512> Text;
+	Record.FormatMessageTo(Text);
+	SerializeToWindow(*Text, Record.GetVerbosity(), Record.GetCategory(), FPlatformTime::ToSeconds64(Record.GetTime().GetCycles()));
+}
+
+void FMacNativeFeedbackContext::SerializeToWindow(const TCHAR* Data, ELogVerbosity::Type Verbosity, const FName& Category, double Time)
+{
+	if (WindowController && bShowingConsoleForSlowTask)
 	{
-		FScopeLock ScopeLock( &CriticalSection );
-		
+		FScopeLock ScopeLock(&CriticalSection);
+
 		static bool Entry=0;
 		if( !GIsCriticalError || Entry )
 		{
@@ -301,7 +317,7 @@ void FMacNativeFeedbackContext::Serialize( const TCHAR* Data, ELogVerbosity::Typ
 			// where each value is either 0 or 1 (can leave off trailing 0's), so
 			// blue on bright yellow is "00101101" and red on black is "1"
 			// An empty string reverts to the normal gray on black
-			
+
 			if (Verbosity == ELogVerbosity::SetColor)
 			{
 				if (FCString::Stricmp(Data, TEXT("")) == 0)
@@ -311,7 +327,7 @@ void FMacNativeFeedbackContext::Serialize( const TCHAR* Data, ELogVerbosity::Typ
 				else
 				{
 					SCOPED_AUTORELEASE_POOL;
-					
+
 					// turn the string into a bunch of 0's and 1's
 					TCHAR String[9];
 					FMemory::Memset(String, 0, sizeof(TCHAR) * UE_ARRAY_COUNT(String));
@@ -320,10 +336,10 @@ void FMacNativeFeedbackContext::Serialize( const TCHAR* Data, ELogVerbosity::Typ
 					{
 						*S -= '0';
 					}
-					
+
 					NSMutableArray* Colors = [[NSMutableArray alloc] init];
 					NSMutableArray* AttributeKeys = [[NSMutableArray alloc] init];
-					
+
 					// Get FOREGROUND_INTENSITY and calculate final color
 					CGFloat Intensity = String[3] ? 1.0 : 0.5;
 					[Colors addObject:[NSColor colorWithSRGBRed:(String[0] ? 1.0 * Intensity : 0.0) green:(String[1] ? 1.0 * Intensity : 0.0) blue:(String[2] ? 1.0 * Intensity : 0.0) alpha:1.0]];
@@ -331,17 +347,17 @@ void FMacNativeFeedbackContext::Serialize( const TCHAR* Data, ELogVerbosity::Typ
 					// Get BACKGROUND_INTENSITY and calculate final color
 					Intensity = String[7] ? 1.0 : 0.5;
 					[Colors addObject:[NSColor colorWithSRGBRed:(String[4] ? 1.0 * Intensity : 0.0) green:(String[5] ? 1.0 * Intensity : 0.0) blue:(String[6] ? 1.0 * Intensity : 0.0) alpha:1.0]];
-					
+
 					[AttributeKeys addObject:NSForegroundColorAttributeName];
 					[AttributeKeys addObject:NSBackgroundColorAttributeName];
-					
+
 					OutstandingTasks++;
 					MainThreadCall(^{
 						if( TextViewTextColor )
 							[TextViewTextColor release];
-						
+
 						TextViewTextColor = [[NSDictionary alloc] initWithObjects:Colors forKeys:AttributeKeys];
-						
+
 						[Colors release];
 						[AttributeKeys release];
 						OutstandingTasks--;
@@ -351,12 +367,12 @@ void FMacNativeFeedbackContext::Serialize( const TCHAR* Data, ELogVerbosity::Typ
 			else
 			{
 				SCOPED_AUTORELEASE_POOL;
-				
+
 				TCHAR OutputString[MAX_SPRINTF]=TEXT(""); //@warning: this is safe as FCString::Sprintf only use 1024 characters max
-				FCString::Sprintf(OutputString,TEXT("%s%s"),*FOutputDeviceHelper::FormatLogLine(Verbosity, Category, Data, GPrintLogTimes),LINE_TERMINATOR);
-				
+				FCString::Sprintf(OutputString,TEXT("%s%s"), *FOutputDeviceHelper::FormatLogLine(Verbosity, Category, Data, GPrintLogTimes, Time), LINE_TERMINATOR);
+
 				CFStringRef CocoaText = FPlatformString::TCHARToCFString(OutputString);
-				
+
 				OutstandingTasks++;
 				MainThreadCall(^{
 					NSAttributedString *AttributedString = [[NSAttributedString alloc] initWithString:(NSString*)CocoaText attributes:TextViewTextColor];
@@ -366,7 +382,7 @@ void FMacNativeFeedbackContext::Serialize( const TCHAR* Data, ELogVerbosity::Typ
 					CFRelease(CocoaText);
 					OutstandingTasks--;
 				}, NSDefaultRunLoopMode, false);
-				
+
 				if(!MacApplication)
 				{
 					FPlatformApplicationMisc::PumpMessages( true );
@@ -379,7 +395,7 @@ void FMacNativeFeedbackContext::Serialize( const TCHAR* Data, ELogVerbosity::Typ
 			try
 			{
 				// Ignore errors to prevent infinite-recursive exception reporting.
-				Serialize( Data, Verbosity, Category );
+				SerializeToWindow(Data, Verbosity, Category, Time);
 			}
 			catch( ... )
 			{}
diff --git a/Engine/Source/Developer/DesktopPlatform/Private/Mac/MacNativeFeedbackContext.h b/Engine/Source/Developer/DesktopPlatform/Private/Mac/MacNativeFeedbackContext.h
--- a/Engine/Source/Developer/DesktopPlatform/Private/Mac/MacNativeFeedbackContext.h
+++ b/Engine/Source/Developer/DesktopPlatform/Private/Mac/MacNativeFeedbackContext.h
@@ -40,7 +40,9 @@ public:
 	FMacNativeFeedbackContext();
 	virtual ~FMacNativeFeedbackContext();
 
-	virtual void Serialize( const TCHAR* V, ELogVerbosity::Type Verbosity, const class FName& Category ) override;
+	virtual void Serialize(const TCHAR* V, ELogVerbosity::Type Verbosity, const FName& Category) override;
+	virtual void Serialize(const TCHAR* V, ELogVerbosity::Type Verbosity, const FName& Category, double Time) override;
+	virtual void SerializeRecord(const UE::FLogRecord& Record) override;
 
 	virtual bool YesNof(const FText& Text) override;
 
@@ -50,10 +52,12 @@ public:
 	virtual void FinalizeSlowTask( ) override;
 	virtual void ProgressReported( const float TotalProgressInterp, FText DisplayMessage ) override;
 
-	FContextSupplier* GetContext() const;
-	void SetContext( FContextSupplier* InSupplier );
+	FContextSupplier* GetContext() const override;
+	void SetContext(FContextSupplier* InContext) override;
 
 private:
+	void SerializeToWindow(const TCHAR* V, ELogVerbosity::Type Verbosity, const FName& Category, double Time);
+
 	void SetDefaultTextColor();
 	
 private:
diff --git a/Engine/Source/Runtime/ApplicationCore/Public/Mac/MacFeedbackContext.h b/Engine/Source/Runtime/ApplicationCore/Public/Mac/MacFeedbackContext.h
--- a/Engine/Source/Runtime/ApplicationCore/Public/Mac/MacFeedbackContext.h
+++ b/Engine/Source/Runtime/ApplicationCore/Public/Mac/MacFeedbackContext.h
@@ -14,76 +14,14 @@
 class FMacFeedbackContext : public FFeedbackContext
 {
 	/** Context information for warning and error messages */
-	FContextSupplier*	Context;
+	FContextSupplier* Context = nullptr;
 
 public:
-	// Constructor.
-	FMacFeedbackContext()
-		: FFeedbackContext()
-		, Context( NULL )
-	{}
-
-	void Serialize( const TCHAR* V, ELogVerbosity::Type Verbosity, const class FName& Category )
-	{
-		// if we set the color for warnings or errors, then reset at the end of the function
-		// note, we have to set the colors directly without using the standard SET_WARN_COLOR macro
-		bool bNeedToResetColor = false;
-		if( Verbosity==ELogVerbosity::Error || Verbosity==ELogVerbosity::Warning )
-		{
-			if( TreatWarningsAsErrors && Verbosity==ELogVerbosity::Warning )
-			{
-				Verbosity = ELogVerbosity::Error;
-			}
-
-			FString Prefix;
-			if( Context )
-			{
-				Prefix = Context->GetContext() + TEXT(" : ");
-			}
-			FString Format = Prefix + FOutputDeviceHelper::FormatLogLine(Verbosity, Category, V);
-
-			if(Verbosity == ELogVerbosity::Error)
-			{
-				Serialize(COLOR_RED, ELogVerbosity::SetColor, Category); bNeedToResetColor = true;
-				// Only store off the message if running a commandlet.
-				if ( IsRunningCommandlet() )
-				{
-					AddError(Format);
-				}
-			}
-			else
-			{
-				Serialize(COLOR_YELLOW, ELogVerbosity::SetColor, Category); bNeedToResetColor = true;
-				// Only store off the message if running a commandlet.
-				if ( IsRunningCommandlet() )
-				{
-					AddWarning(Format);
-				}
-			}
-		}
-
-		if( GLogConsole && IsRunningCommandlet() && !GLog->IsRedirectingTo(GLogConsole) )
-			GLogConsole->Serialize( V, Verbosity, Category );
-		if( !GLog->IsRedirectingTo( this ) )
-			GLog->Serialize( V, Verbosity, Category );
-
-		if (bNeedToResetColor)
-		{
-			Serialize(COLOR_NONE, ELogVerbosity::SetColor, Category);
-		}
-	}
-
-	bool YesNof( const FText& Question )
+	bool YesNof(const FText& Question) override
 	{
 		return FPlatformMisc::MessageBoxExt(EAppMsgType::YesNo, *Question.ToString(), TEXT("")) == EAppReturnType::Yes;
 	}
-	FContextSupplier* GetContext() const
-	{
-		return Context;
-	}
-	void SetContext( FContextSupplier* InSupplier )
-	{
-		Context = InSupplier;
-	}
-};
 
+	FContextSupplier* GetContext() const override { return Context; }
+	void SetContext(FContextSupplier* InContext) override { Context = InContext; }
+};
