Fixed a missing include in MacNativeFeedbackContext.cpp

#rnx

diff --git a/Engine/Source/Developer/DesktopPlatform/Private/Mac/MacNativeFeedbackContext.cpp b/Engine/Source/Developer/DesktopPlatform/Private/Mac/MacNativeFeedbackContext.cpp
--- a/Engine/Source/Developer/DesktopPlatform/Private/Mac/MacNativeFeedbackContext.cpp
+++ b/Engine/Source/Developer/DesktopPlatform/Private/Mac/MacNativeFeedbackContext.cpp
@@ -1,6 +1,7 @@
 // Copyright Epic Games, Inc. All Rights Reserved.
 
 #include "MacNativeFeedbackContext.h"
+#include "HAL/PlatformTime.h"
 #include "Logging/StructuredLog.h"
 #include "Mac/MacApplication.h"
 #include "Mac/CocoaThread.h"
