Fix unicode character being used instead of ' in team's

diff --git a/Engine/Source/Developer/DesktopPlatform/Private/DesktopPlatformBase.cpp b/Engine/Source/Developer/DesktopPlatform/Private/DesktopPlatformBase.cpp
--- a/Engine/Source/Developer/DesktopPlatform/Private/DesktopPlatformBase.cpp
+++ b/Engine/Source/Developer/DesktopPlatform/Private/DesktopPlatformBase.cpp
@@ -924,7 +924,7 @@ bool FDesktopPlatformBase::GetOidcTokenStatus(const FString& RootDir, const FStr
 bool FDesktopPlatformBase::GetOidcAccessTokenInteractive(const FString& RootDir,  const FString& Arguments, FFeedbackContext* Warn, int32& OutReturnCode)
 {
 	FText OidcInteractivePromptTitle = NSLOCTEXT("OidcToken", "OidcToken_InteractiveLaunchPromptTitle", "Unreal Engine - Authentication Required");
-	FText OidcInteractiveLaunchPromptText = NSLOCTEXT("OidcToken", "OidcToken_InteractiveLaunch", "Your team’s preferred DDC (Distributed Data Cache) requires you to log in. Click OK to open the authentication page in your web browser.\n\nYou can cancel authentication and work with a different shared or local DDC instead. However, this may cause delays while the editor prepares the assets you need.");
+	FText OidcInteractiveLaunchPromptText = NSLOCTEXT("OidcToken", "OidcToken_InteractiveLaunch", "Your team's preferred DDC (Distributed Data Cache) requires you to log in. Click OK to open the authentication page in your web browser.\n\nYou can cancel authentication and work with a different shared or local DDC instead. However, this may cause delays while the editor prepares the assets you need.");
 	EAppReturnType::Type userAcknowledgedResult = FPlatformMisc::MessageBoxExt(EAppMsgType::OkCancel, *OidcInteractiveLaunchPromptText.ToString(), *OidcInteractivePromptTitle.ToString());
 
 	if (userAcknowledgedResult != EAppReturnType::Ok)
