Fix for installed builds on linux not having a stable UUID from machine to machine

Note: Since Linux has no launcher, EnumerateLauncherEngineInstallations doesn't take care of this for us in Epic created Linux installed builds

diff --git a/Engine/Source/Developer/DesktopPlatform/Private/Linux/DesktopPlatformLinux.cpp b/Engine/Source/Developer/DesktopPlatform/Private/Linux/DesktopPlatformLinux.cpp
--- a/Engine/Source/Developer/DesktopPlatform/Private/Linux/DesktopPlatformLinux.cpp
+++ b/Engine/Source/Developer/DesktopPlatform/Private/Linux/DesktopPlatformLinux.cpp
@@ -135,7 +135,27 @@ bool FDesktopPlatformLinux::RegisterEngineInstallation(const FString &RootDir, F
 		ConfigFile.Read(ConfigPath);
 
 		FConfigSection &Section = ConfigFile.FindOrAdd(TEXT("Installations"));
-		OutIdentifier = FGuid::NewGuid().ToString(EGuidFormats::DigitsWithHyphens);
+
+		// If this is an installed build, use that Guid instead of generating a new one
+		FString InstallationIdPath = FString(RootDir / "Engine" / "Build" / "InstalledBuild.txt");
+		FArchive* File = IFileManager::Get().CreateFileReader(*InstallationIdPath, FILEREAD_NoFail);
+		if(File)
+		{
+			FFileHelper::LoadFileToString(OutIdentifier, *File);
+			FGuid GuidCheck(OutIdentifier);
+			if(!GuidCheck.IsValid())
+			{
+				OutIdentifier = FGuid::NewGuid().ToString(EGuidFormats::DigitsWithHyphens);
+			}
+
+			File->Close();
+			delete File;
+		}
+		else
+		{
+			OutIdentifier = FGuid::NewGuid().ToString(EGuidFormats::DigitsWithHyphens);
+		}
+
 		Section.AddUnique(*OutIdentifier, RootDir);
 
 		ConfigFile.Dirty = true;
diff --git a/Engine/Source/Programs/AutomationTool/Scripts/FinalizeInstalledBuild.cs b/Engine/Source/Programs/AutomationTool/Scripts/FinalizeInstalledBuild.cs
--- a/Engine/Source/Programs/AutomationTool/Scripts/FinalizeInstalledBuild.cs
+++ b/Engine/Source/Programs/AutomationTool/Scripts/FinalizeInstalledBuild.cs
@@ -35,7 +35,7 @@ namespace AutomationTool
 			// Write InstalledBuild.txt to indicate Engine is installed
 			string InstalledBuildFile = CommandUtils.CombinePaths(OutputDir, "Engine/Build/InstalledBuild.txt");
 			CommandUtils.CreateDirectory(CommandUtils.GetDirectoryName(InstalledBuildFile));
-			CommandUtils.WriteAllText(InstalledBuildFile, "");
+			CommandUtils.WriteAllText(InstalledBuildFile, Guid.NewGuid().ToString().ToUpper());
 
 			// Write InstalledBuild.txt to indicate Engine is installed
 			string Project = ParseParamValue("Project");
