OidcToken - Improved error message when running unattended to clarify that you can login using the editor.

diff --git a/Engine/Source/Developer/DesktopPlatform/Private/DesktopPlatformBase.cpp b/Engine/Source/Developer/DesktopPlatform/Private/DesktopPlatformBase.cpp
--- a/Engine/Source/Developer/DesktopPlatform/Private/DesktopPlatformBase.cpp
+++ b/Engine/Source/Developer/DesktopPlatform/Private/DesktopPlatformBase.cpp
@@ -840,7 +840,7 @@ bool FDesktopPlatformBase::GetOidcAccessToken(const FString& RootDir, const FStr
 		}
 		else
 		{
-			UE_LOG(LogDesktopPlatform, Error, TEXT("Unable to allocate an access token. Unattended set so unable to request interactive login. Make sure you are logged in using UGS or using the UGS cli command 'login'. Provider used: '%s'. Ran OidcToken (project file is '%s', exe path is '%s')"), *ProviderIdentifier, *ProjectFileName, *GetOidcTokenExecutableFilename(RootDir));
+			UE_LOG(LogDesktopPlatform, Error, TEXT("Unable to allocate an access token. Unattended set so unable to request interactive login. Make sure you start the editor and login once or log in using UGS or using the UGS cli command 'login'. Provider used: '%s'. Ran OidcToken (project file is '%s', exe path is '%s')"), *ProviderIdentifier, *ProjectFileName, *GetOidcTokenExecutableFilename(RootDir));
 			return false;
 		}
 	}
