Build break fix - Changed DesktopPlatform/LockFile to be unimplemented on non-desktop platforms instead of being a compile time error for them as DesktopPlatform is getting built for more than just Windows, Linux, and Mac.

diff --git a/Engine/Source/Developer/DesktopPlatform/Private/LockFile.cpp b/Engine/Source/Developer/DesktopPlatform/Private/LockFile.cpp
--- a/Engine/Source/Developer/DesktopPlatform/Private/LockFile.cpp
+++ b/Engine/Source/Developer/DesktopPlatform/Private/LockFile.cpp
@@ -104,7 +104,7 @@ FLockFile::TryReadAndClear(const TCHAR* FileName, FString& OutContents)
 
 	return true;
 #else
-	#error Not implemented for this platform
+	unimplemented();
 	return false;
 #endif
 }
\ No newline at end of file
