Error -> Warning when we are unable to resolve a token for Cloud DDC. This will not cause the build to fail, but it will typically be slower for it so error is harsh (as that in itself causes the build to fail)

diff --git a/Engine/Source/Developer/DerivedDataCache/Private/HttpCacheStore.cpp b/Engine/Source/Developer/DerivedDataCache/Private/HttpCacheStore.cpp
--- a/Engine/Source/Developer/DerivedDataCache/Private/HttpCacheStore.cpp
+++ b/Engine/Source/Developer/DerivedDataCache/Private/HttpCacheStore.cpp
@@ -1763,7 +1763,7 @@ bool FHttpCacheStore::AcquireAccessToken(IHttpClient* Client)
 			return true;
 		}
 
-		UE_LOG(LogDerivedDataCache, Error, TEXT("%s: OidcToken: Failed to log in to HTTP services."), *Domain);
+		UE_LOG(LogDerivedDataCache, Warning, TEXT("%s: OidcToken: Failed to log in to HTTP services."), *Domain);
 		FailedLoginAttempts++;
 		return false;
 	}
diff --git a/Engine/Source/Developer/DesktopPlatform/Private/DesktopPlatformBase.cpp b/Engine/Source/Developer/DesktopPlatform/Private/DesktopPlatformBase.cpp
--- a/Engine/Source/Developer/DesktopPlatform/Private/DesktopPlatformBase.cpp
+++ b/Engine/Source/Developer/DesktopPlatform/Private/DesktopPlatformBase.cpp
@@ -840,7 +840,7 @@ bool FDesktopPlatformBase::GetOidcAccessToken(const FString& RootDir, const FStr
 		}
 		else
 		{
-			UE_LOG(LogDesktopPlatform, Error, TEXT("Unable to allocate an access token. Unattended set so unable to request interactive login. Make sure you start the editor and login once or log in using UGS or using the UGS cli command 'login'. Provider used: '%s'. Ran OidcToken (project file is '%s', exe path is '%s')"), *ProviderIdentifier, *ProjectFileName, *GetOidcTokenExecutableFilename(RootDir));
+			UE_LOG(LogDesktopPlatform, Warning, TEXT("Unable to allocate an access token. Unattended set so unable to request interactive login. Make sure you start the editor and login once or log in using UGS or using the UGS cli command 'login'. Provider used: '%s'. Ran OidcToken (project file is '%s', exe path is '%s')"), *ProviderIdentifier, *ProjectFileName, *GetOidcTokenExecutableFilename(RootDir));
 			return false;
 		}
 	}
