Remove a debugging check that was erroneously checked in as part of fix for UE-171276

diff --git a/Engine/Source/Developer/DesktopPlatform/Private/Linux/DesktopPlatformLinux.cpp b/Engine/Source/Developer/DesktopPlatform/Private/Linux/DesktopPlatformLinux.cpp
--- a/Engine/Source/Developer/DesktopPlatform/Private/Linux/DesktopPlatformLinux.cpp
+++ b/Engine/Source/Developer/DesktopPlatform/Private/Linux/DesktopPlatformLinux.cpp
@@ -138,7 +138,7 @@ bool FDesktopPlatformLinux::RegisterEngineInstallation(const FString &RootDir, F
 
 		// If this is an installed build, use that Guid instead of generating a new one
 		FString InstallationIdPath = FString(RootDir / "Engine" / "Build" / "InstalledBuild.txt");
-		FArchive* File = IFileManager::Get().CreateFileReader(*InstallationIdPath, FILEREAD_NoFail);
+		FArchive* File = IFileManager::Get().CreateFileReader(*InstallationIdPath, FILEREAD_Silent);
 		if(File)
 		{
 			FFileHelper::LoadFileToString(OutIdentifier, *File);
