Use relative paths in TargetInfo.json to make them portable between machines (eg. when including in a zip of PCBs).

diff --git a/Engine/Source/Developer/DesktopPlatform/Private/DesktopPlatformBase.cpp b/Engine/Source/Developer/DesktopPlatform/Private/DesktopPlatformBase.cpp
--- a/Engine/Source/Developer/DesktopPlatform/Private/DesktopPlatformBase.cpp
+++ b/Engine/Source/Developer/DesktopPlatform/Private/DesktopPlatformBase.cpp
@@ -1778,6 +1778,7 @@ bool FDesktopPlatformBase::ReadTargetInfo(const FString& FileName, TArray<FTarge
 	Targets.SetNum(TargetArray->Num());
 
 	// Parse the entries
+	FString BaseDir = FPaths::GetPath(FileName);
 	for (int Idx = 0; Idx < TargetArray->Num(); Idx++)
 	{
 		const FJsonValue& TargetValue = *(*TargetArray)[Idx].Get();
@@ -1801,6 +1802,11 @@ bool FDesktopPlatformBase::ReadTargetInfo(const FString& FileName, TArray<FTarge
 		{
 			return false;
 		}
+
+		FString Path = FPaths::ConvertRelativePathToFull(FPaths::Combine(BaseDir, Targets[Idx].Path));
+		FPaths::MakePlatformFilename(Path);
+
+		Targets[Idx].Path = Path;
 	}
 
 	return true;
diff --git a/Engine/Source/Programs/UnrealBuildTool/Modes/QueryTargetsMode.cs b/Engine/Source/Programs/UnrealBuildTool/Modes/QueryTargetsMode.cs
--- a/Engine/Source/Programs/UnrealBuildTool/Modes/QueryTargetsMode.cs
+++ b/Engine/Source/Programs/UnrealBuildTool/Modes/QueryTargetsMode.cs
@@ -143,11 +143,17 @@ namespace UnrealBuildTool
 							continue;
 						}
 					}
-					
+
+					// Get the path to the target
+					FileReference? path = Assembly.GetTargetFileName(TargetName);
+
 					// Write the target info
 					Writer.WriteObjectStart();
 					Writer.WriteValue("Name", TargetName);
-					Writer.WriteValue("Path", Assembly.GetTargetFileName(TargetName)?.ToString());
+					if (path != null)
+					{
+						Writer.WriteValue("Path", path.MakeRelativeTo(OutputFile.Directory));
+					}
 					Writer.WriteValue("Type", TargetRules.Type.ToString());
 					Writer.WriteObjectEnd();
 				}
