Add SVG icon for uprojects and ubinaries, which scales down to small icon scales without truncating.

diff --git a/Engine/Source/Developer/DesktopPlatform/Private/Linux/DesktopPlatformLinux.cpp b/Engine/Source/Developer/DesktopPlatform/Private/Linux/DesktopPlatformLinux.cpp
--- a/Engine/Source/Developer/DesktopPlatform/Private/Linux/DesktopPlatformLinux.cpp
+++ b/Engine/Source/Developer/DesktopPlatform/Private/Linux/DesktopPlatformLinux.cpp
@@ -384,7 +384,7 @@ bool FDesktopPlatformLinux::UpdateFileAssociations()
 		return true;
 	}
 
-	// Install the icons, one for uprojects and one for the main Unreal Engine launcher.
+	// Install the png icons, one for uprojects and one for the main Unreal Engine launcher.
 	if (!RunXDGUtil(FString::Printf(TEXT("xdg-icon-resource install --novendor --mode user --context mimetypes --size 256 %sPrograms/UnrealVersionSelector/Private/Linux/Resources/Icon.png uproject"), *FPaths::EngineSourceDir())))
 	{
 		return false;
@@ -395,6 +395,38 @@ bool FDesktopPlatformLinux::UpdateFileAssociations()
 		return false;
 	}
 
+	FString IconSource = FPaths::Combine(FPaths::EngineSourceDir(), TEXT("Programs/UnrealVersionSelector/Private/Linux/Resources/Icon.svg")); 
+	FString ProjectIconDestination = FPaths::Combine(FPlatformMisc::GetEnvironmentVariable(TEXT("HOME")), TEXT(".local/share/icons/hicolor/scalable/mimetypes")); 
+	FString BinaryIconDestination  = FPaths::Combine(FPlatformMisc::GetEnvironmentVariable(TEXT("HOME")), TEXT(".local/share/icons/hicolor/scalable/apps")); 
+
+	IFileManager& FileManager = IFileManager::Get();
+
+	// Ensure that the proper directories exist for svg icons as well
+	if (!FileManager.DirectoryExists(*ProjectIconDestination))
+	{
+		if (!FileManager.MakeDirectory(*ProjectIconDestination))
+		{
+			return false;
+		}
+	}
+	if (!FileManager.DirectoryExists(*BinaryIconDestination))
+	{
+		if (!FileManager.MakeDirectory(*BinaryIconDestination))
+		{
+			return false;
+		}
+	}
+
+	// Install the svg icons; this is done manually because xdg-icon-resource doesn't support installation of svg files
+	if (FileManager.Copy(*FPaths::Combine(ProjectIconDestination, TEXT("uproject.svg")), *IconSource) != COPY_OK)
+	{
+		return false;
+	}
+	if (FileManager.Copy(*FPaths::Combine(BinaryIconDestination, TEXT("ubinary.svg")), *IconSource) != COPY_OK)
+	{
+		return false;
+	}
+
 	FString AbsoluteEngineDir = IFileManager::Get().ConvertToAbsolutePathForExternalAppForWrite(*FPaths::EngineDir());
 
 	// Add the desktop file for the Unreal Version Selector mime-type from the template.
