Fixed a compile error caused by the removal of FLogTime::GetCycles()

#rnx

diff --git a/Engine/Source/Developer/DesktopPlatform/Private/Mac/MacNativeFeedbackContext.cpp b/Engine/Source/Developer/DesktopPlatform/Private/Mac/MacNativeFeedbackContext.cpp
--- a/Engine/Source/Developer/DesktopPlatform/Private/Mac/MacNativeFeedbackContext.cpp
+++ b/Engine/Source/Developer/DesktopPlatform/Private/Mac/MacNativeFeedbackContext.cpp
@@ -301,7 +301,7 @@ void FMacNativeFeedbackContext::SerializeRecord(const UE::FLogRecord& Record)
 
 	TStringBuilder<512> Text;
 	Record.FormatMessageTo(Text);
-	SerializeToWindow(*Text, Record.GetVerbosity(), Record.GetCategory(), FPlatformTime::ToSeconds64(Record.GetTime().GetCycles()));
+	SerializeToWindow(*Text, Record.GetVerbosity(), Record.GetCategory(), -1.0);
 }
 
 void FMacNativeFeedbackContext::SerializeToWindow(const TCHAR* Data, ELogVerbosity::Type Verbosity, const FName& Category, double Time)
