OidcToken: Delete temporary files from Intermediate directory after we no longer need them.
#rnx

diff --git a/Engine/Source/Developer/DesktopPlatform/Private/DesktopPlatformBase.cpp b/Engine/Source/Developer/DesktopPlatform/Private/DesktopPlatformBase.cpp
--- a/Engine/Source/Developer/DesktopPlatform/Private/DesktopPlatformBase.cpp
+++ b/Engine/Source/Developer/DesktopPlatform/Private/DesktopPlatformBase.cpp
@@ -9,6 +9,7 @@
 #include "Misc/FeedbackContext.h"
 #include "Misc/App.h"
 #include "Misc/EngineVersion.h"
+#include "Misc/ScopeExit.h"
 #include "Serialization/JsonTypes.h"
 #include "Serialization/JsonReader.h"
 #include "Serialization/JsonSerializer.h"
@@ -808,6 +809,10 @@ bool FDesktopPlatformBase::IsUnrealBuildToolRunning()
 bool FDesktopPlatformBase::GetOidcAccessToken(const FString& RootDir, const FString& ProjectFileName, const FString& ProviderIdentifier, bool bUnattended, FFeedbackContext* Warn, FString& OutToken, FDateTime& OutTokenExpiresAt, bool& bOutWasInteractiveLogin)
 {
 	FString ResultFilePath = FPaths::CreateTempFilename(*FPaths::ProjectIntermediateDir(), TEXT("oidcToken.json"));
+	ON_SCOPE_EXIT
+	{
+		IFileManager::Get().Delete(*ResultFilePath, false /* RequireExists */, true /* EvenIfReadOnly */, true /* Quiet */);
+	};
 
 	FString Arguments = TEXT(" ");
 	Arguments += FString::Printf(TEXT(" --Service=\"%s\""), *ProviderIdentifier);
