QOL: Enable the Slate based file browser on Windows when using the NullPlatform

diff --git a/Engine/Source/Developer/DesktopPlatform/DesktopPlatform.Build.cs b/Engine/Source/Developer/DesktopPlatform/DesktopPlatform.Build.cs
--- a/Engine/Source/Developer/DesktopPlatform/DesktopPlatform.Build.cs
+++ b/Engine/Source/Developer/DesktopPlatform/DesktopPlatform.Build.cs
@@ -10,7 +10,8 @@ public class DesktopPlatform : ModuleRules
 			new string[] {
 				"DerivedDataCache",
 				"SlateFontDialog",
-				"LauncherPlatform",
+				"SlateFileDialogs",
+				"LauncherPlatform"
 			}
 		);
 
@@ -18,19 +19,18 @@ public class DesktopPlatform : ModuleRules
 			new string[] {
 				"Core",
 				"ApplicationCore",
-				"Json",
+				"Json"
 			}
 		);
+		
 
 		if (Target.IsInPlatformGroup(UnrealPlatformGroup.Linux))
 		{
 			if(Target.Type == TargetType.Editor)
 			{
-				PrivateIncludePathModuleNames.Add("SlateFontDialog");
 				DynamicallyLoadedModuleNames.Add("SlateFontDialog");
 			}
-
-			PrivateIncludePathModuleNames.Add("SlateFileDialogs");
+			
 			DynamicallyLoadedModuleNames.Add("SlateFileDialogs");
 
 			AddEngineThirdPartyPrivateStaticDependencies(Target, "SDL2");
diff --git a/Engine/Source/Developer/DesktopPlatform/Private/DesktopPlatformModule.cpp b/Engine/Source/Developer/DesktopPlatform/Private/DesktopPlatformModule.cpp
--- a/Engine/Source/Developer/DesktopPlatform/Private/DesktopPlatformModule.cpp
+++ b/Engine/Source/Developer/DesktopPlatform/Private/DesktopPlatformModule.cpp
@@ -2,13 +2,21 @@
 
 #include "DesktopPlatformModule.h"
 #include "DesktopPlatformPrivate.h"
+#include "Null/NullPlatformApplicationMisc.h"
 
 IMPLEMENT_MODULE( FDesktopPlatformModule, DesktopPlatform );
 DEFINE_LOG_CATEGORY(LogDesktopPlatform);
 
 void FDesktopPlatformModule::StartupModule()
 {
-	DesktopPlatform = new FDesktopPlatform();
+	if(FNullPlatformApplicationMisc::IsUsingNullApplication())
+	{
+		DesktopPlatform = new FDesktopPlatformNull();
+	}
+	else
+	{
+		DesktopPlatform = new FDesktopPlatform();
+	}
 
 	FPlatformMisc::SetEnvironmentVar(TEXT("UE_DesktopUnrealProcess"), TEXT("1"));
 }
diff --git a/Engine/Source/Developer/DesktopPlatform/Private/DesktopPlatformPrivate.h b/Engine/Source/Developer/DesktopPlatform/Private/DesktopPlatformPrivate.h
--- a/Engine/Source/Developer/DesktopPlatform/Private/DesktopPlatformPrivate.h
+++ b/Engine/Source/Developer/DesktopPlatform/Private/DesktopPlatformPrivate.h
@@ -14,6 +14,8 @@
 #else
 #include "DesktopPlatformStub.h"
 #endif
+
+#include "Null/DesktopPlatformNull.h"
 // IWYU pragma: end_exports
 
 DECLARE_LOG_CATEGORY_EXTERN(LogDesktopPlatform, Log, All);
diff --git a/Engine/Source/Developer/DesktopPlatform/Private/Null/DesktopPlatformNull.cpp b/Engine/Source/Developer/DesktopPlatform/Private/Null/DesktopPlatformNull.cpp
new file mode 100644
--- /dev/null
+++ b/Engine/Source/Developer/DesktopPlatform/Private/Null/DesktopPlatformNull.cpp
@@ -0,0 +1,238 @@
+// Copyright Epic Games, Inc. All Rights Reserved.
+
+#include "DesktopPlatformPrivate.h"
+#include "HAL/FileManager.h"
+#include "HAL/ThreadHeartBeat.h"
+#include "Internationalization/Regex.h"
+#include "ISlateFileDialogModule.h"
+#include "Misc/ConfigCacheIni.h"
+#include "Misc/FeedbackContext.h"
+#include "Misc/FeedbackContextMarkup.h"
+#include "Misc/FileHelper.h"
+#include "Misc/Guid.h"
+#include "Misc/Paths.h"
+#include "Modules/ModuleManager.h"
+#include "Null/DesktopPlatformNull.h"
+
+#if PLATFORM_LINUX
+#include "ISlateFontDialogModule.h" 
+#endif
+
+#define LOCTEXT_NAMESPACE "DesktopPlatform"
+#define MAX_FILETYPES_STR 4096
+#define MAX_FILENAME_STR 65536
+
+DEFINE_LOG_CATEGORY_STATIC(LogDesktopPlatformNull, Log, All);
+
+FDesktopPlatformNull::FDesktopPlatformNull()
+	:	FDesktopPlatformBase()
+{
+}
+
+FDesktopPlatformNull::~FDesktopPlatformNull()
+{
+}
+
+bool FDesktopPlatformNull::OpenFileDialog(const void* ParentWindowHandle, const FString& DialogTitle, const FString& DefaultPath, const FString& DefaultFile, const FString& FileTypes, uint32 Flags, TArray<FString>& OutFilenames, int32& OutFilterIndex)
+{
+	if (!FModuleManager::Get().IsModuleLoaded("SlateFileDialogs"))
+	{
+		FModuleManager::Get().LoadModule("SlateFileDialogs");
+	}
+
+	ISlateFileDialogsModule *FileDialog = FModuleManager::GetModulePtr<ISlateFileDialogsModule>("SlateFileDialogs");
+
+	if (FileDialog)
+	{
+		return FileDialog->OpenFileDialog(ParentWindowHandle, DialogTitle, DefaultPath, DefaultFile, FileTypes, Flags, OutFilenames, OutFilterIndex);
+	}
+
+	return false;
+}
+
+bool FDesktopPlatformNull::OpenFileDialog(const void* ParentWindowHandle, const FString& DialogTitle, const FString& DefaultPath, const FString& DefaultFile, const FString& FileTypes, uint32 Flags, TArray<FString>& OutFilenames)
+{
+	if (!FModuleManager::Get().IsModuleLoaded("SlateFileDialogs"))
+	{
+		FModuleManager::Get().LoadModule("SlateFileDialogs");
+	}
+
+	ISlateFileDialogsModule *FileDialog = FModuleManager::GetModulePtr<ISlateFileDialogsModule>("SlateFileDialogs");
+
+	if (FileDialog)
+	{
+		return FileDialog->OpenFileDialog(ParentWindowHandle, DialogTitle, DefaultPath, DefaultFile, FileTypes, Flags, OutFilenames);
+	}
+
+	return false;
+}
+
+bool FDesktopPlatformNull::SaveFileDialog(const void* ParentWindowHandle, const FString& DialogTitle, const FString& DefaultPath, const FString& DefaultFile, const FString& FileTypes, uint32 Flags, TArray<FString>& OutFilenames)
+{
+	if (!FModuleManager::Get().IsModuleLoaded("SlateFileDialogs"))
+	{
+		FModuleManager::Get().LoadModule("SlateFileDialogs");
+	}
+
+	ISlateFileDialogsModule *FileDialog = FModuleManager::GetModulePtr<ISlateFileDialogsModule>("SlateFileDialogs");
+
+	if (FileDialog)
+	{
+		return FileDialog->SaveFileDialog(ParentWindowHandle, DialogTitle, DefaultPath, DefaultFile, FileTypes, Flags, OutFilenames);
+	}
+
+	return false;
+}
+
+bool FDesktopPlatformNull::OpenDirectoryDialog(const void* ParentWindowHandle, const FString& DialogTitle, const FString& DefaultPath, FString& OutFolderName)
+{
+	if (!FModuleManager::Get().IsModuleLoaded("SlateFileDialogs"))
+	{
+		FModuleManager::Get().LoadModule("SlateFileDialogs");
+	}
+
+	ISlateFileDialogsModule *FileDialog = FModuleManager::GetModulePtr<ISlateFileDialogsModule>("SlateFileDialogs");
+
+	if (FileDialog)
+	{
+		return FileDialog->OpenDirectoryDialog(ParentWindowHandle, DialogTitle, DefaultPath, OutFolderName);
+	}
+
+	return false;
+}
+
+bool FDesktopPlatformNull::OpenFontDialog(const void* ParentWindowHandle, FString& OutFontName, float& OutHeight, EFontImportFlags& OutFlags)
+{
+#if PLATFORM_LINUX
+	if (!FModuleManager::Get().IsModuleLoaded("SlateFontDialog"))
+	{
+		FModuleManager::Get().LoadModule("SlateFontDialog");
+	}
+
+	ISlateFontDialogModule* FontDialog = FModuleManager::GetModulePtr<ISlateFontDialogModule>("SlateFontDialog");
+	
+	if (FontDialog)
+	{
+		return FontDialog->OpenFontDialog(OutFontName, OutHeight, OutFlags);
+	}
+	else
+	{
+		UE_LOG(LogDesktopPlatformNull, Warning, TEXT("Error reading results of font dialog"));
+	}
+#endif
+	return false;
+}
+
+bool FDesktopPlatformNull::FileDialogShared(bool bSave, const void* ParentWindowHandle, const FString& DialogTitle, const FString& DefaultPath, const FString& DefaultFile, const FString& FileTypes, uint32 Flags, TArray<FString>& OutFilenames, int32& OutFilterIndex)
+{
+	return false;
+}
+
+bool FDesktopPlatformNull::RegisterEngineInstallation(const FString &RootDir, FString &OutIdentifier)
+{
+	bool bRes = false;
+	if (IsValidRootDirectory(RootDir))
+	{
+		FConfigFile ConfigFile;
+		FString ConfigPath = FString(FPlatformProcess::ApplicationSettingsDir()) / FString(TEXT("UnrealEngine")) / FString(TEXT("Install.ini"));
+		ConfigFile.Read(ConfigPath);
+
+		FConfigSection &Section = ConfigFile.FindOrAdd(TEXT("Installations"));
+
+		// If this is an installed build, use that Guid instead of generating a new one
+		FString InstallationIdPath = FString(RootDir / "Engine" / "Build" / "InstalledBuild.txt");
+		FArchive* File = IFileManager::Get().CreateFileReader(*InstallationIdPath, FILEREAD_Silent);
+		if(File)
+		{
+			FFileHelper::LoadFileToString(OutIdentifier, *File);
+			FGuid GuidCheck(OutIdentifier);
+			if(!GuidCheck.IsValid())
+			{
+				OutIdentifier = FGuid::NewGuid().ToString(EGuidFormats::DigitsWithHyphens);
+			}
+
+			File->Close();
+			delete File;
+		}
+		else
+		{
+			OutIdentifier = FGuid::NewGuid().ToString(EGuidFormats::DigitsWithHyphens);
+		}
+
+		Section.AddUnique(*OutIdentifier, RootDir);
+
+		ConfigFile.Dirty = true;
+		ConfigFile.Write(ConfigPath);
+		bRes = true;
+	}
+	return bRes;
+}
+
+void FDesktopPlatformNull::EnumerateEngineInstallations(TMap<FString, FString> &OutInstallations)
+{
+	EnumerateLauncherEngineInstallations(OutInstallations);
+	// TODO (William.Belcher): See if we actually need to run this when using the null platform
+}
+
+bool FDesktopPlatformNull::IsSourceDistribution(const FString &RootDir)
+{
+	// Use the default test
+	return FDesktopPlatformBase::IsSourceDistribution(RootDir);
+}
+
+bool FDesktopPlatformNull::VerifyFileAssociations()
+{
+	return true;
+}
+
+bool FDesktopPlatformNull::UpdateFileAssociations()
+{
+	return true;
+}
+
+bool FDesktopPlatformNull::OpenProject(const FString &ProjectFileName)
+{
+	// Get the project filename in a native format
+	FString PlatformProjectFileName = ProjectFileName;
+	FPaths::MakePlatformFilename(PlatformProjectFileName);
+
+	STUBBED("FDesktopPlatformNull::OpenProject");
+	return false;
+}
+
+bool FDesktopPlatformNull::RunUnrealBuildTool(const FText& Description, const FString& RootDir, const FString& Arguments, FFeedbackContext* Warn, int32& OutExitCode)
+{
+	OutExitCode = 1;
+
+	// Get the path to UBT
+	FString UnrealBuildToolPath = GetUnrealBuildToolExecutableFilename(RootDir);
+	if(IFileManager::Get().FileSize(*UnrealBuildToolPath) < 0)
+	{
+		Warn->Logf(ELogVerbosity::Error, TEXT("Couldn't find UnrealBuildTool at '%s'"), *UnrealBuildToolPath);
+		return false;
+	}
+
+	// Write the output
+	Warn->Logf(TEXT("Running %s %s"), *UnrealBuildToolPath, *Arguments);
+
+	return FFeedbackContextMarkup::PipeProcessOutput(Description, UnrealBuildToolPath, Arguments, Warn, &OutExitCode) && OutExitCode == 0;
+}
+
+FFeedbackContext* FDesktopPlatformNull::GetNativeFeedbackContext()
+{
+	//unimplemented();
+	STUBBED("FDesktopPlatformNull::GetNativeFeedbackContext");
+	return GWarn;
+}
+
+FString FDesktopPlatformNull::GetUserTempPath()
+{
+	return FString(FPlatformProcess::UserTempDir());
+}
+
+FString FDesktopPlatformNull::GetOidcTokenExecutableFilename(const FString& RootDir) const
+{	
+	return "";
+}
+
+#undef LOCTEXT_NAMESPACE
diff --git a/Engine/Source/Developer/DesktopPlatform/Private/Null/DesktopPlatformNull.h b/Engine/Source/Developer/DesktopPlatform/Private/Null/DesktopPlatformNull.h
new file mode 100644
--- /dev/null
+++ b/Engine/Source/Developer/DesktopPlatform/Private/Null/DesktopPlatformNull.h
@@ -0,0 +1,43 @@
+// Copyright Epic Games, Inc. All Rights Reserved.
+
+#pragma once
+
+#include "CoreMinimal.h"
+#include "DesktopPlatformBase.h"
+
+class FDesktopPlatformNull : public FDesktopPlatformBase
+{
+public:
+	// IDesktopPlatform Implementation
+	virtual bool OpenFileDialog(const void* ParentWindowHandle, const FString& DialogTitle, const FString& DefaultPath, const FString& DefaultFile, const FString& FileTypes, uint32 Flags, TArray<FString>& OutFilenames) override;
+	virtual bool OpenFileDialog(const void* ParentWindowHandle, const FString& DialogTitle, const FString& DefaultPath, const FString& DefaultFile, const FString& FileTypes, uint32 Flags, TArray<FString>& OutFilenames, int32& OutFilterIndex) override;
+	virtual bool SaveFileDialog(const void* ParentWindowHandle, const FString& DialogTitle, const FString& DefaultPath, const FString& DefaultFile, const FString& FileTypes, uint32 Flags, TArray<FString>& OutFilenames) override;
+	virtual bool OpenDirectoryDialog(const void* ParentWindowHandle, const FString& DialogTitle, const FString& DefaultPath, FString& OutFolderName) override;
+	virtual bool OpenFontDialog(const void* ParentWindowHandle, FString& OutFontName, float& OutHeight, EFontImportFlags& OutFlags) override;
+
+	virtual bool RegisterEngineInstallation(const FString &RootDir, FString &OutIdentifier) override;
+	virtual void EnumerateEngineInstallations(TMap<FString, FString> &OutInstallations) override;
+
+	virtual bool IsSourceDistribution(const FString &RootDir) override;
+
+	virtual bool VerifyFileAssociations() override;
+	virtual bool UpdateFileAssociations() override;
+
+	virtual bool OpenProject(const FString &ProjectFileName) override;
+
+	using FDesktopPlatformBase::RunUnrealBuildTool;
+	virtual bool RunUnrealBuildTool(const FText& Description, const FString& RootDir, const FString& Arguments, FFeedbackContext* Warn, int32& OutExitCode) override;
+
+	virtual FFeedbackContext* GetNativeFeedbackContext() override;
+
+	virtual FString GetUserTempPath() override;
+
+	FDesktopPlatformNull();
+	virtual ~FDesktopPlatformNull();
+
+protected:
+	virtual FString GetOidcTokenExecutableFilename(const FString& RootDir) const override;
+
+private:
+	bool FileDialogShared(bool bSave, const void* ParentWindowHandle, const FString& DialogTitle, const FString& DefaultPath, const FString& DefaultFile, const FString& FileTypes, uint32 Flags, TArray<FString>& OutFilenames, int32& OutFilterIndex);
+};
diff --git a/Engine/Source/Developer/DirectoryWatcher/Private/DirectoryWatcherPrivate.h b/Engine/Source/Developer/DirectoryWatcher/Private/DirectoryWatcherPrivate.h
--- a/Engine/Source/Developer/DirectoryWatcher/Private/DirectoryWatcherPrivate.h
+++ b/Engine/Source/Developer/DirectoryWatcher/Private/DirectoryWatcherPrivate.h
@@ -16,5 +16,5 @@ DECLARE_LOG_CATEGORY_EXTERN(LogDirectoryWatcher, Log, All);
 #include "Linux/DirectoryWatchRequestLinux.h"
 #include "Linux/DirectoryWatcherLinux.h"
 #else
-#error "Unknown platform"
+#include "DirectoryWatcherStub.h"
 #endif
diff --git a/Engine/Source/Developer/DirectoryWatcher/Private/DirectoryWatcherStub.h b/Engine/Source/Developer/DirectoryWatcher/Private/DirectoryWatcherStub.h
new file mode 100644
--- /dev/null
+++ b/Engine/Source/Developer/DirectoryWatcher/Private/DirectoryWatcherStub.h
@@ -0,0 +1,21 @@
+// Copyright Epic Games, Inc. All Rights Reserved.
+
+#pragma once
+
+#include "CoreMinimal.h"
+#include "IDirectoryWatcher.h"
+
+
+class FDirectoryWatcherStub : public IDirectoryWatcher
+{
+public:
+	FDirectoryWatcherStub();
+	virtual ~FDirectoryWatcherStub();
+
+	virtual bool RegisterDirectoryChangedCallback_Handle (const FString& Directory, const FDirectoryChanged& InDelegate, FDelegateHandle& OutHandle, uint32 Flags) override { return false; }
+	virtual bool UnregisterDirectoryChangedCallback_Handle(const FString& Directory, FDelegateHandle InHandle) override { return false; }
+	virtual void Tick (float DeltaSeconds) override { }
+	virtual bool DumpStats() override { return false; }
+};
+
+typedef FDirectoryWatcherStub FDirectoryWatcher;
