Add a setting that controls whether the modal acknowledgement dialog is shown before a browser is launched for OIDC authentication in editor.  This setting is editor-wide (not per project), but:

1. A project can override the default value in that project's DefaultEditorSettings.ini
2. A user can indicate they don't want to see the dialog in the future, and thata setting is persisted to the user's editor settings file

If the setting indicates (either through defaults or user choice) that the acknowledgement dialog can be skipped, then unattended operations (eg: commandlets like cook etc..) are permitted to attempt interactive login by launching the browser.  If the browser login succeeds within a fixed timespan (30 seconds), then the unattended operation can proceed having successfully obtained an oidc token.  If the browser login hasn't succeeded in the fixed timespan, then the unanttended operation will proceed but will have failed to obtain an oidc token.  In neither case do we block indefinitely on completion of browser login to obtain an oidc token.

diff --git a/Engine/Source/Developer/DesktopPlatform/Private/DesktopPlatformBase.cpp b/Engine/Source/Developer/DesktopPlatform/Private/DesktopPlatformBase.cpp
--- a/Engine/Source/Developer/DesktopPlatform/Private/DesktopPlatformBase.cpp
+++ b/Engine/Source/Developer/DesktopPlatform/Private/DesktopPlatformBase.cpp
@@ -832,21 +832,20 @@ bool FDesktopPlatformBase::GetOidcAccessToken(const FString& RootDir, const FStr
 
 	if (ExitCode == 10)
 	{
-		if (!bUnattended)
-		{
-			bRes = GetOidcAccessTokenInteractive(RootDir, Arguments, Warn, ExitCode);
+		bRes = GetOidcAccessTokenInteractive(RootDir, Arguments, bUnattended, Warn, ExitCode);
 
-			bOutWasInteractiveLogin = true;
+		bOutWasInteractiveLogin = true;
 
-			if (!bRes)
+		if (!bRes)
+		{
+			if (bUnattended)
+			{
+				UE_LOG(LogDesktopPlatform, Warning, TEXT("Unable to allocate an access token. Unattended set so unable to complete interactive login. Make sure you start the editor and login once or log in using UGS or using the UGS cli command 'login'. Provider used: '%s'. Ran OidcToken (project file is '%s', exe path is '%s')"), *ProviderIdentifier, *ProjectFileName, *GetOidcTokenExecutableFilename(RootDir));
+			}
+			else
 			{
 				UE_LOG(LogDesktopPlatform, Error, TEXT("Unable to allocate an access token. Interactive login failed, make sure you are assigned access and are able to login in the created browser window. Provider used: '%s'. Ran OidcToken (project file is '%s', exe path is '%s')"), *ProviderIdentifier, *ProjectFileName, *GetOidcTokenExecutableFilename(RootDir));
-				return false;
 			}
-		}
-		else
-		{
-			UE_LOG(LogDesktopPlatform, Warning, TEXT("Unable to allocate an access token. Unattended set so unable to request interactive login. Make sure you start the editor and login once or log in using UGS or using the UGS cli command 'login'. Provider used: '%s'. Ran OidcToken (project file is '%s', exe path is '%s')"), *ProviderIdentifier, *ProjectFileName, *GetOidcTokenExecutableFilename(RootDir));
 			return false;
 		}
 	}
@@ -938,13 +937,35 @@ bool FDesktopPlatformBase::GetOidcTokenStatus(const FString& RootDir, const FStr
 	return false;
 }
 
-bool FDesktopPlatformBase::GetOidcAccessTokenInteractive(const FString& RootDir,  const FString& Arguments, FFeedbackContext* Warn, int32& OutReturnCode)
+bool FDesktopPlatformBase::GetOidcAccessTokenInteractive(const FString& RootDir, const FString& Arguments, bool bUnattended, FFeedbackContext* Warn, int32& OutReturnCode)
 {
-	FText OidcInteractivePromptTitle = NSLOCTEXT("OidcToken", "OidcToken_InteractiveLaunchPromptTitle", "Unreal Engine - Authentication Required");
-	FText OidcInteractiveLaunchPromptText = NSLOCTEXT("OidcToken", "OidcToken_InteractiveLaunch", "Your team's preferred DDC (Derived Data Cache) requires you to log in. Click OK to open the authentication page in your web browser.\n\nYou can cancel authentication and work with a different shared or local DDC instead. However, this may cause delays while the editor prepares the assets you need.");
-	EAppReturnType::Type userAcknowledgedResult = FPlatformMisc::MessageBoxExt(EAppMsgType::OkCancel, *OidcInteractiveLaunchPromptText.ToString(), *OidcInteractivePromptTitle.ToString());
+	EAppReturnType::Type userAcknowledgedResult = EAppReturnType::Yes;
+
+	check(GConfig && GConfig->IsReadyForUse());
+	
+	bool bSkipInitialAcknowledgement = false;
+	GConfig->GetBool(TEXT("/Script/UnrealEd.EditorSettings"), TEXT("InteractiveOidcWithoutAcknowledgement"), bSkipInitialAcknowledgement, GEditorSettingsIni);
+
+	if (!bSkipInitialAcknowledgement)
+	{
+		if (bUnattended)
+		{
+			OutReturnCode = -1;
+			return false;
+		}
+
+		FText OidcInteractivePromptTitle = NSLOCTEXT("OidcToken", "OidcToken_InteractiveLaunchPromptTitle", "Unreal Engine - Authentication Required");
+		FText OidcInteractiveLaunchPromptText = NSLOCTEXT("OidcToken", "OidcToken_InteractiveLaunch", "Your team's preferred DDC (Derived Data Cache) requires you to log in.\n\nClick Yes to open the authentication page in your web browser, or Yes To All to always proceed to browser authentication without a prompt.\n\nClick No to decline authentication and work with a different shared or local DDC instead. However, this may cause delays while the editor prepares the assets you need.");
+		userAcknowledgedResult = FPlatformMisc::MessageBoxExt(EAppMsgType::YesNoYesAll, *OidcInteractiveLaunchPromptText.ToString(), *OidcInteractivePromptTitle.ToString());
+
+		if (userAcknowledgedResult == EAppReturnType::YesAll)
+		{
+			GConfig->SetBool(TEXT("/Script/UnrealEd.EditorSettings"), TEXT("InteractiveOidcWithoutAcknowledgement"), true, GEditorSettingsIni);
+			userAcknowledgedResult = EAppReturnType::Yes;
+		}
+	}
 
-	if (userAcknowledgedResult != EAppReturnType::Ok)
+	if (userAcknowledgedResult != EAppReturnType::Yes)
 	{
 		OutReturnCode = -1;
 		return false;
@@ -991,9 +1012,15 @@ bool FDesktopPlatformBase::GetOidcAccessTokenInteractive(const FString& RootDir,
 			UE_LOG(LogDesktopPlatform, Display, TEXT("Waiting for OidcToken to finish login..."));
 			DurationPhase = EWaitDurationPhase::Prompt;
 		}
-		// once we have waited for 30 seconds without success we give the user a option to abort
+		// once we have waited for 30 seconds without success we give the user a option to abort, or if unattended, abort without prompting
 		else if (WaitDuration > 30.0 && DurationPhase == EWaitDurationPhase::Prompt)
 		{
+			if (bUnattended)
+			{
+				bIsFinished = !FPlatformProcess::IsProcRunning(ProcHandle);
+				break;
+			}
+
 			FText OidcLongWaitPromptTitle = NSLOCTEXT("OidcToken", "OidcToken_LongWaitPromptTitle", "Wait for user login?");
 			FText OidcLongWaitPromptText = NSLOCTEXT("OidcToken", "OidcToken_LongWaitPromptText", "Login is taking a long time, make sure you have entered your credentials in your browser window. It can be in a tab in an already existing window. Keep waiting?");
 			if (FPlatformMisc::MessageBoxExt(EAppMsgType::YesNo, *OidcLongWaitPromptText.ToString(), *OidcLongWaitPromptTitle.ToString()) == EAppReturnType::No)
diff --git a/Engine/Source/Developer/DesktopPlatform/Private/DesktopPlatformBase.h b/Engine/Source/Developer/DesktopPlatform/Private/DesktopPlatformBase.h
--- a/Engine/Source/Developer/DesktopPlatform/Private/DesktopPlatformBase.h
+++ b/Engine/Source/Developer/DesktopPlatform/Private/DesktopPlatformBase.h
@@ -85,7 +85,7 @@ private:
 
 	bool InvokeOidcTokenToolSync(const FText& Description, const FString& RootDir, const FString& Arguments, FFeedbackContext* Warn, int32& OutReturnCode, FString& OutProcOutput);
 	FProcHandle InvokeOidcTokenToolAsync(const FString& InArguments, void*& OutReadPipe, void*& OutWritePipe);
-	bool GetOidcAccessTokenInteractive(const FString& RootDir,  const FString& Arguments, FFeedbackContext* Warn, int32& OutReturnCode);
+	bool GetOidcAccessTokenInteractive(const FString& RootDir, const FString& Arguments, bool bUnattended, FFeedbackContext* Warn, int32& OutReturnCode);
 
 	static bool ReadTargetInfo(const FString& FileName, TArray<FTargetInfo>& Targets);
 
