Linux: Fix IsVanilla check by implementing EnumerateLauncherEngineInstallations for Linux

- Also add -BuildId flag to allow released builds to be registered as "launcher" builds.
- This gets us telemetry and the ability to open projects made on the Mac and Windows builds without being prompted to open a copy

diff --git a/Engine/Build/InstalledEngineBuild.xml b/Engine/Build/InstalledEngineBuild.xml
--- a/Engine/Build/InstalledEngineBuild.xml
+++ b/Engine/Build/InstalledEngineBuild.xml
@@ -56,6 +56,16 @@
 	<!-- The analytics type -->
 	<Option Name="AnalyticsTypeOverride" DefaultValue="" Description="Identifier for analytic events to send"/>
 
+	<!-- The build identifier -->
+	<Option Name="BuildIdOverride" DefaultValue="" Description="Override unique identifier for this build"/>
+	<Property Name="BuildId" Value="" />
+	<Property Name="BuildId" Value="$(BuildIdOverride)" If="'$(BuildIdOverride)' != ''"/>
+
+	<!--- Create command line argument for build identifier -->
+	<Property Name="BuildIdentifierArg" Value="" />
+	<Property Name="BuildIdentifierArg" Value="-BuildId=&quot;$(BuildId)&quot;" If="'$(BuildId)' != ''"/>
+
+
 	<!-- Whether to embed source server info (ie. P4 paths) into PDB files -->
 	<Option Name="EmbedSrcSrvInfo" Restrict="true|false" DefaultValue="false" Description="Whether to add Source indexing to Windows game apps so they can be added to a symbol server"/>
 
@@ -1313,7 +1323,7 @@
 			<Property Name="ArchitectureCommands" Value="$(ArchitectureCommands) -AndroidArchitectures=$(AndroidArchitectures)" If="'$(WithAndroid)' == true"/>
 
 			<!-- Finalize the build and then tag the text file that this process creates -->
-			<Command Name="FinalizeInstalledBuild" Arguments="-OutputDir=$(LocalInstalledDirLinux) -Platforms=$(Platforms) -ContentOnlyPlatforms=$(ContentOnlyPlatforms) $(ArchitectureCommands) -AnalyticsTypeOverride=$(AnalyticsTypeOverride)"/>
+			<Command Name="FinalizeInstalledBuild" Arguments="-OutputDir=$(LocalInstalledDirLinux) $(BuildIdentifierArg) -Platforms=$(Platforms) -ContentOnlyPlatforms=$(ContentOnlyPlatforms) $(ArchitectureCommands) -AnalyticsTypeOverride=$(AnalyticsTypeOverride)"/>
 			<Tag Files="$(LocalInstalledDirLinux)/Engine/Build/InstalledBuild.txt" With="#Installed Build Linux Files"/>
 
 			<!-- Sanitize all the receipts in the target directory -->
diff --git a/Engine/Source/Developer/DesktopPlatform/Private/Linux/DesktopPlatformLinux.cpp b/Engine/Source/Developer/DesktopPlatform/Private/Linux/DesktopPlatformLinux.cpp
--- a/Engine/Source/Developer/DesktopPlatform/Private/Linux/DesktopPlatformLinux.cpp
+++ b/Engine/Source/Developer/DesktopPlatform/Private/Linux/DesktopPlatformLinux.cpp
@@ -142,8 +142,9 @@ bool FDesktopPlatformLinux::RegisterEngineInstallation(const FString &RootDir, F
 		if(File)
 		{
 			FFileHelper::LoadFileToString(OutIdentifier, *File);
+			OutIdentifier.TrimEndInline();
 			FGuid GuidCheck(OutIdentifier);
-			if(!GuidCheck.IsValid())
+			if(!GuidCheck.IsValid() && !OutIdentifier.StartsWith(TEXT("UE_")))
 			{
 				OutIdentifier = FGuid::NewGuid().ToString(EGuidFormats::DigitsWithHyphens);
 			}
@@ -157,6 +158,7 @@ bool FDesktopPlatformLinux::RegisterEngineInstallation(const FString &RootDir, F
 		}
 
 		Section.AddUnique(*OutIdentifier, RootDir);
+		OutIdentifier.RemoveFromStart(TEXT("UE_"));
 
 		ConfigFile.Dirty = true;
 		ConfigFile.Write(ConfigPath);
@@ -192,9 +194,9 @@ void FDesktopPlatformLinux::EnumerateEngineInstallations(TMap<FString, FString>
 	TArray<FName> KeysToRemove;
 	for (auto It : Section)
 	{
-		const FString& EngineDir = It.Value.GetValue();
+		const FString& RootDir = It.Value.GetValue();
 		// We remove entries pointing to a folder that doesn't exist or was using the wrong path.
-		if (EngineDir.Contains(FPaths::EngineDir()) || !IFileManager::Get().DirectoryExists(*EngineDir))
+		if (RootDir.Contains(FPaths::EngineDir()) || !IFileManager::Get().DirectoryExists(*RootDir))
 		{
 			KeysToRemove.Add(It.Key);
 			ConfigFile.Dirty = true;
@@ -210,15 +212,15 @@ void FDesktopPlatformLinux::EnumerateEngineInstallations(TMap<FString, FString>
 	// Iterate through all entries.
 	for (auto It : Section)
 	{
-		FString EngineDir = It.Value.GetValue();
-		FPaths::NormalizeDirectoryName(EngineDir);
-		FPaths::CollapseRelativeDirectories(EngineDir);
+		FString NormalizedRootDir = It.Value.GetValue();
+		FPaths::NormalizeDirectoryName(NormalizedRootDir);
+		FPaths::CollapseRelativeDirectories(NormalizedRootDir);
 
 		FString EngineId;
-		const FName* Key = Section.FindKey(EngineDir);
+		const FName* Key = Section.FindKey(NormalizedRootDir);
 		if (Key == nullptr)
 		{
-			Key = SectionsToAdd.FindKey(EngineDir);
+			Key = SectionsToAdd.FindKey(NormalizedRootDir);
 		}
 
 		if (Key)
@@ -229,17 +231,19 @@ void FDesktopPlatformLinux::EnumerateEngineInstallations(TMap<FString, FString>
 		}
 		else
 		{
-			if (!OutInstallations.FindKey(EngineDir))
+			if (!OutInstallations.FindKey(NormalizedRootDir))
 			{
 				EngineId = FGuid::NewGuid().ToString(EGuidFormats::DigitsWithHyphens);
-				SectionsToAdd.AddUnique(*EngineId, EngineDir);
+				SectionsToAdd.AddUnique(*EngineId, NormalizedRootDir);
 
 				ConfigFile.Dirty = true;
 			}
 		}
-		if (!EngineId.IsEmpty() && !OutInstallations.Find(EngineId))
+
+		// local builds only, don't add entries that start with "UE_" which signifies a released build
+		if (!EngineId.IsEmpty() && !OutInstallations.Find(EngineId) && !EngineId.StartsWith(TEXT("UE_")))
 		{
-			OutInstallations.Add(EngineId, EngineDir);
+			OutInstallations.Add(EngineId, NormalizedRootDir);
 		}
 	}
 
@@ -253,6 +257,63 @@ void FDesktopPlatformLinux::EnumerateEngineInstallations(TMap<FString, FString>
 	IFileManager::Get().Delete(*UProjectPath);
 }
 
+void FDesktopPlatformLinux::EnumerateLauncherEngineInstallations(TMap<FString, FString> &OutInstallations)
+{
+	FConfigFile ConfigFile;
+	FString ConfigPath = FString(FPlatformProcess::ApplicationSettingsDir()) / FString(TEXT("UnrealEngine")) / FString(TEXT("Install.ini"));
+	ConfigFile.Read(ConfigPath);
+	
+	FConfigSection &Section = ConfigFile.FindOrAdd(TEXT("Installations"));
+
+	FString InstallationIdPath = FString(FPaths::EngineDir() / "Build" / "InstalledBuild.txt");
+	FArchive* File = IFileManager::Get().CreateFileReader(*InstallationIdPath, FILEREAD_Silent);
+	if (File)
+	{
+		FString NormalizedRootDir = FPaths::RootDir();
+		FPaths::NormalizeDirectoryName(NormalizedRootDir);
+		FPaths::CollapseRelativeDirectories(NormalizedRootDir);
+
+		FString Id;
+		FFileHelper::LoadFileToString(Id, *File);
+		Id.TrimEndInline();
+
+		// if the user unzipped a new installed build into a previously registered directory, we need to fix up the key
+		const FName* OldKey = Section.FindKey(NormalizedRootDir);
+		if(OldKey && Id != OldKey->ToString())
+		{
+			Section.Remove(*OldKey);
+			Section.Add(*Id, *NormalizedRootDir);
+			ConfigFile.Dirty = true;
+			ConfigFile.Write(ConfigPath);
+		}
+		File->Close();
+		delete File;
+	}
+
+	// now fill OutInstallations with only released builds
+	for (auto It : Section)
+	{
+		const FString RootDir = It.Value.GetValue();
+		FString GuidOrId = It.Key.ToString();
+
+		// We skip entries pointing to a folder that doesn't exist or was using the wrong path.
+		if (RootDir.Contains(FPaths::EngineDir()) || !IFileManager::Get().DirectoryExists(*RootDir))
+		{
+			continue;
+		}
+
+		// released builds only, add entries starting with "UE_"
+		if(GuidOrId.RemoveFromStart(TEXT("UE_"), ESearchCase::CaseSensitive))
+		{
+			if(!OutInstallations.Contains(GuidOrId))
+			{
+				OutInstallations.Add(*GuidOrId, *RootDir);
+			}
+		}
+	}
+
+}
+
 bool FDesktopPlatformLinux::IsSourceDistribution(const FString &RootDir)
 {
 	// Check for the existence of a GenerateProjectFiles.sh file. This allows compatibility with the GitHub 4.0 release.
diff --git a/Engine/Source/Developer/DesktopPlatform/Private/Linux/DesktopPlatformLinux.h b/Engine/Source/Developer/DesktopPlatform/Private/Linux/DesktopPlatformLinux.h
--- a/Engine/Source/Developer/DesktopPlatform/Private/Linux/DesktopPlatformLinux.h
+++ b/Engine/Source/Developer/DesktopPlatform/Private/Linux/DesktopPlatformLinux.h
@@ -17,6 +17,7 @@ public:
 
 	virtual bool RegisterEngineInstallation(const FString &RootDir, FString &OutIdentifier) override;
 	virtual void EnumerateEngineInstallations(TMap<FString, FString> &OutInstallations) override;
+	virtual void EnumerateLauncherEngineInstallations(TMap<FString, FString> &OutInstallations) override;
 
 	virtual bool IsSourceDistribution(const FString &RootDir) override;
 
diff --git a/Engine/Source/Programs/AutomationTool/BuildGraph/Scripts/InstalledBuild.cs b/Engine/Source/Programs/AutomationTool/BuildGraph/Scripts/InstalledBuild.cs
--- a/Engine/Source/Programs/AutomationTool/BuildGraph/Scripts/InstalledBuild.cs
+++ b/Engine/Source/Programs/AutomationTool/BuildGraph/Scripts/InstalledBuild.cs
@@ -126,6 +126,7 @@ namespace AutomationTool
 
 			BgStringOption CrashReporterAPIURL = new BgStringOption("CrashReporterAPIURL", "The URL to use to talk to the CrashReporterClient API.", "");
 			BgStringOption CrashReporterAPIKey = new BgStringOption("CrashReporterAPIKey", "The API key to use to talk to the CrashReporterClient API.", "");
+			BgStringOption BuildId = new BgStringOption("BuildId", "The unique build identifier to associate with this installed build", "");
 
 			BgString CrashReporterCompileArgs = "";
 			CrashReporterCompileArgs = CrashReporterCompileArgs.If(CrashReporterAPIURL != "" & CrashReporterAPIKey != "", BgString.Format("-define:CRC_TELEMETRY_URL=\"{0}\" -define:CRC_TELEMETRY_KEY_DEV=\"{1}\" -define:CRC_TELEMETRY_KEY_RELEASE=\"{1}\" -OverrideBuildEnvironment", CrashReporterAPIURL, CrashReporterAPIKey));
diff --git a/Engine/Source/Programs/AutomationTool/Scripts/FinalizeInstalledBuild.cs b/Engine/Source/Programs/AutomationTool/Scripts/FinalizeInstalledBuild.cs
--- a/Engine/Source/Programs/AutomationTool/Scripts/FinalizeInstalledBuild.cs
+++ b/Engine/Source/Programs/AutomationTool/Scripts/FinalizeInstalledBuild.cs
@@ -20,6 +20,7 @@ namespace AutomationTool
 	[Help("<Platform>Architectures=<ArchitectureList>)", "List of architectures that are used for a given platform (optional)")]
 	[Help("<Platform>GPUArchitectures=<GPUArchitectureList>)", "List of GPU architectures that are used for a given platform (optional)")]
 	[Help("AnalyticsTypeOverride=<TypeName>", "Name to give this build for analytics purposes (optional)")]
+	[Help("BuildId=<BuildIdentifier>", "Unique identifier for this build (optional)")]
 	class FinalizeInstalledBuild : BuildCommand
 	{
 		/// <summary>
@@ -31,11 +32,18 @@ namespace AutomationTool
 			List<UnrealTargetPlatform> Platforms = ParsePlatformsParamValue("Platforms");
 			List<UnrealTargetPlatform> ContentOnlyPlatforms = ParsePlatformsParamValue("ContentOnlyPlatforms");
 			string AnalyticsTypeOverride = ParseParamValue("AnalyticsTypeOverride");
+			string BuildId = ParseOptionalStringParam("BuildId");
 
 			// Write InstalledBuild.txt to indicate Engine is installed
 			string InstalledBuildFile = CommandUtils.CombinePaths(OutputDir, "Engine/Build/InstalledBuild.txt");
 			CommandUtils.CreateDirectory(CommandUtils.GetDirectoryName(InstalledBuildFile));
-			CommandUtils.WriteAllText(InstalledBuildFile, Guid.NewGuid().ToString().ToUpper());
+
+			// Write build Identifier into InstalledBuild.txt
+			if(String.IsNullOrWhiteSpace(BuildId))
+			{
+				BuildId = Guid.NewGuid().ToString().ToUpper();
+			}
+			CommandUtils.WriteAllText(InstalledBuildFile, BuildId);
 
 			// Write InstalledBuild.txt to indicate Engine is installed
 			string Project = ParseParamValue("Project");
@@ -191,4 +199,4 @@ namespace AutomationTool
 			}
 		}
 	}
-}
\ No newline at end of file
+}
