Fix TargetInfo validation.
Fix checking in plugin directories for targets.
Prevents running UBT on every launch.

diff --git a/Engine/Source/Developer/DesktopPlatform/Private/DesktopPlatformBase.cpp b/Engine/Source/Developer/DesktopPlatform/Private/DesktopPlatformBase.cpp
--- a/Engine/Source/Developer/DesktopPlatform/Private/DesktopPlatformBase.cpp
+++ b/Engine/Source/Developer/DesktopPlatform/Private/DesktopPlatformBase.cpp
@@ -19,6 +19,7 @@
 #include "Misc/OutputDeviceRedirector.h"
 #include "Misc/ScopedSlowTask.h"
 #include "Misc/SecureHash.h"
+#include "String/RemoveFrom.h"
 
 #define LOCTEXT_NAMESPACE "DesktopPlatform"
 
@@ -1127,30 +1128,21 @@ struct FTargetFileVisitor : IPlatformFile::FDirectoryStatVisitor
 	{
 	}
 
-	virtual bool Visit(const TCHAR* FileNameOrDirectory, const FFileStatData& StatData) override
+	virtual bool Visit(const TCHAR* InFileNameOrDirectory, const FFileStatData& StatData) override
 	{
 		if (StatData.bIsDirectory)
 		{
-			SubDirectories.Add(FileNameOrDirectory);
+			SubDirectories.Add(InFileNameOrDirectory);
 			return true;
 		}
+		
+		FStringView FileNameOrDirectory(InFileNameOrDirectory);
 
 		// NOTE: This code needs to behave the same as FindAllRulesSourceFiles in Rules.cs
-		static const TCHAR TargetExt[] = TEXT(".Target.cs");
-		static const int32 TargetExtLen = UE_ARRAY_COUNT(TargetExt) - 1;
-		static const TCHAR ModuleExt[] = TEXT(".Build.cs");
-		static const int32 ModuleExtLen = UE_ARRAY_COUNT(ModuleExt) - 1;
-		static const TCHAR AutomationCsprojExt[] = TEXT(".automation.csproj");
-		static const int32 AutomationCsprojExtLen = UE_ARRAY_COUNT(AutomationCsprojExt) - 1;
-		static const TCHAR UBTCsprojExt[] = TEXT(".ubtplugin.csproj");
-		static const int32 UBTCsprojExtLen = UE_ARRAY_COUNT(UBTCsprojExt) - 1;
-		static const TCHAR UBTIgnoreExt[] = TEXT(".ubtignore");
-		static const int32 UBTIgnoreExtLen = UE_ARRAY_COUNT(UBTIgnoreExt) - 1;
-
-		int32 Length = FCString::Strlen(FileNameOrDirectory);
-		if (Length > TargetExtLen && FCString::Stricmp(FileNameOrDirectory + Length - TargetExtLen, TargetExt) == 0)
+		if (FStringView WithoutExtension = UE::String::RemoveFromEnd(FileNameOrDirectory, TEXTVIEW(".Target.cs"));
+			WithoutExtension.Len() != FileNameOrDirectory.Len())
 		{
-			FString TargetName = FPaths::GetCleanFilename(FString(Length - TargetExtLen, FileNameOrDirectory));
+			FString TargetName = FPaths::GetCleanFilename(FString(WithoutExtension));
 
 			// skip target rules that are platform extension or platform group specializations
 			// Matches logic found in QueryTargetsMode.cs WriteTargetInfo
@@ -1160,24 +1152,17 @@ struct FTargetFileVisitor : IPlatformFile::FDirectoryStatVisitor
 				return true;
 			}
 
-			return (StatData.ModificationTime < MaxDateTime && RemainingTargetNames.Remove(TargetName) == 1);
-		}
-		else if (Length > ModuleExtLen && FCString::Stricmp(FileNameOrDirectory + Length - ModuleExtLen, ModuleExt) == 0)
-		{
-			bSearchSubDirectories = false;
-			return true;
-		}
-		else if (Length > AutomationCsprojExtLen && FCString::Stricmp(FileNameOrDirectory + Length - AutomationCsprojExtLen, AutomationCsprojExt) == 0)
-		{
-			bSearchSubDirectories = false;
-			return true;
-		}
-		else if (Length > UBTCsprojExtLen && FCString::Stricmp(FileNameOrDirectory + Length - UBTCsprojExtLen, UBTCsprojExt) == 0)
-		{
-			bSearchSubDirectories = false;
-			return true;
+			if (StatData.ModificationTime < MaxDateTime)
+			{	
+				RemainingTargetNames.Remove(TargetName);
+			}
+			return RemainingTargetNames.Num() != 0;
 		}
-		else if (Length > UBTIgnoreExtLen && FCString::Stricmp(FileNameOrDirectory + Length - UBTIgnoreExtLen, UBTIgnoreExt) == 0)
+		else if (FileNameOrDirectory.EndsWith(TEXTVIEW(".Build.cs"))
+			|| FileNameOrDirectory.EndsWith(TEXTVIEW(".automation.csproj"))
+		 	|| FileNameOrDirectory.EndsWith(TEXTVIEW(".ubtplugin.csproj"))
+			|| FileNameOrDirectory.EndsWith(TEXTVIEW(".ubtignore"))
+		)
 		{
 			bSearchSubDirectories = false;
 			return true;
@@ -1189,6 +1174,7 @@ struct FTargetFileVisitor : IPlatformFile::FDirectoryStatVisitor
 
 bool IsTargetInfoValid(const TArray<FTargetInfo>& Targets, TArray<FString>& DirectoryNames, const FDateTime& LastModifiedTime)
 {
+	TRACE_CPUPROFILER_EVENT_SCOPE_STR("IsTargetInfoValid");
 	if (FApp::GetEngineIsPromotedBuild())
 	{
 		// Promoted builds may not have source code, so we will assume all supplied targets are valid since they will not appear on disk
@@ -1208,9 +1194,10 @@ bool IsTargetInfoValid(const TArray<FTargetInfo>& Targets, TArray<FString>& Dire
 	for(int Idx = 0; Idx < DirectoryNames.Num(); Idx++)
 	{
 		FTargetFileVisitor Visitor(OriginalTargetNames, RemainingTargetNames, LastModifiedTime);
-		if(!IFileManager::Get().IterateDirectoryStat(*DirectoryNames[Idx], Visitor))
+		IFileManager::Get().IterateDirectoryStat(*DirectoryNames[Idx], Visitor);
+		if (RemainingTargetNames.Num() == 0)
 		{
-			return false;
+			return true;
 		}
 		if(Visitor.bSearchSubDirectories)
 		{
@@ -1224,6 +1211,7 @@ bool IsTargetInfoValid(const TArray<FTargetInfo>& Targets, TArray<FString>& Dire
 
 const TArray<FTargetInfo>& FDesktopPlatformBase::GetTargetsForProject(const FString& ProjectFile) const
 {
+	TRACE_CPUPROFILER_EVENT_SCOPE_STR("FDesktopPlatfomrBase::GetTargetsForProject");
 	// Normalize the project filename
 	FString NormalizedProjectFile = ProjectFile;
 	FPaths::NormalizeFilename(NormalizedProjectFile);
@@ -1257,7 +1245,7 @@ const TArray<FTargetInfo>& FDesktopPlatformBase::GetTargetsForProject(const FStr
 	{
 		// Read it in and check it's still valid
 		TArray<FTargetInfo> NewTargets;
-		TArray<FString> DirectoryNames = { ProjectSourceDir, ProjectDir / TEXT("Platforms"), ProjectDir / TEXT("Restricted") };
+		TArray<FString> DirectoryNames = { ProjectSourceDir, ProjectDir / TEXT("Platforms"), ProjectDir / TEXT("Restricted"), ProjectDir / TEXT("Plugins") };
 		if(ReadTargetInfo(InfoFileName, NewTargets) && IsTargetInfoValid(NewTargets, DirectoryNames, StatData.ModificationTime))
 		{
 			return ProjectFileToTargets.Emplace(MoveTemp(NormalizedProjectFile), MoveTemp(NewTargets));
