Fix for the case I broke previously when trying to avoid executing UBT to find target definitions, where we find a target file that was not in the cache.
Copy more code from UBT to more correctly scan required plugin directories.

diff --git a/Engine/Source/Developer/DesktopPlatform/Private/DesktopPlatformBase.cpp b/Engine/Source/Developer/DesktopPlatform/Private/DesktopPlatformBase.cpp
--- a/Engine/Source/Developer/DesktopPlatform/Private/DesktopPlatformBase.cpp
+++ b/Engine/Source/Developer/DesktopPlatform/Private/DesktopPlatformBase.cpp
@@ -2,6 +2,7 @@
 
 #include "DesktopPlatformBase.h"
 #include "HAL/FileManager.h"
+#include "Logging/LogScopedVerbosityOverride.h"
 #include "Misc/FileHelper.h"
 #include "Misc/Paths.h"
 #include "Misc/Guid.h"
@@ -1118,13 +1119,14 @@ struct FTargetFileVisitor : IPlatformFile::FDirectoryStatVisitor
 	TSet<FString>& RemainingTargetNames;
 	FDateTime MaxDateTime;
 	TArray<FString> SubDirectories;
-	bool bSearchSubDirectories;
+	bool bSearchSubDirectories = true;
+	bool bSearchPluginForTargets = false;
+	bool bCacheInvalid = false;
 
 	FTargetFileVisitor(const TSet<FString>& InOriginalTargetNames, TSet<FString>& InRemainingTargetNames, FDateTime InMaxDateTime)
 		: OriginalTargetNames(InOriginalTargetNames)
 		, RemainingTargetNames(InRemainingTargetNames)
 		, MaxDateTime(InMaxDateTime)
-		, bSearchSubDirectories(true)
 	{
 	}
 
@@ -1152,11 +1154,24 @@ struct FTargetFileVisitor : IPlatformFile::FDirectoryStatVisitor
 				return true;
 			}
 
-			if (StatData.ModificationTime < MaxDateTime)
+			if (StatData.ModificationTime >= MaxDateTime)
 			{	
-				RemainingTargetNames.Remove(TargetName);
+				UE_LOG(LogDesktopPlatform, Log, TEXT("Found target file %s newer than cache"), InFileNameOrDirectory);
+				bCacheInvalid = true;
 			}
-			return RemainingTargetNames.Num() != 0;
+			if (RemainingTargetNames.Remove(TargetName) != 1)
+			{
+				UE_LOG(LogDesktopPlatform, Log, TEXT("Found target file %s not present in cache"), InFileNameOrDirectory);
+				bCacheInvalid = true;
+			}
+			
+			return !bCacheInvalid;
+		}
+		else if (FileNameOrDirectory.EndsWith(TEXTVIEW(".uplugin")))
+		{
+			bSearchSubDirectories = false;
+			bSearchPluginForTargets = true;
+			return true;
 		}
 		else if (FileNameOrDirectory.EndsWith(TEXTVIEW(".Build.cs"))
 			|| FileNameOrDirectory.EndsWith(TEXTVIEW(".automation.csproj"))
@@ -1172,6 +1187,7 @@ struct FTargetFileVisitor : IPlatformFile::FDirectoryStatVisitor
 	}
 };
 
+// Note: This function must find all target files found by CreateProjectRulesAssembly in RulesCompiler.cs
 bool IsTargetInfoValid(const TArray<FTargetInfo>& Targets, TArray<FString>& DirectoryNames, const FDateTime& LastModifiedTime)
 {
 	TRACE_CPUPROFILER_EVENT_SCOPE_STR("IsTargetInfoValid");
@@ -1180,7 +1196,7 @@ bool IsTargetInfoValid(const TArray<FTargetInfo>& Targets, TArray<FString>& Dire
 		// Promoted builds may not have source code, so we will assume all supplied targets are valid since they will not appear on disk
 		return true;
 	}
-
+	
 	// Create the state 
 	TSet<FString> RemainingTargetNames;
 	for (const FTargetInfo& Target : Targets)
@@ -1190,21 +1206,44 @@ bool IsTargetInfoValid(const TArray<FTargetInfo>& Targets, TArray<FString>& Dire
 
 	TSet<FString> OriginalTargetNames = RemainingTargetNames;
 
+	IFileManager& FM = IFileManager::Get();
 	// Loop through all the directories
 	for(int Idx = 0; Idx < DirectoryNames.Num(); Idx++)
 	{
+		// UE_LOG(LogDesktopPlatform, Verbose, TEXT("Checking %s for target files"), *DirectoryNames[Idx]);
 		FTargetFileVisitor Visitor(OriginalTargetNames, RemainingTargetNames, LastModifiedTime);
-		IFileManager::Get().IterateDirectoryStat(*DirectoryNames[Idx], Visitor);
+		FM.IterateDirectoryStat(*DirectoryNames[Idx], Visitor);
+		if (Visitor.bCacheInvalid)
+		{
+			return false;
+		}
 		if (RemainingTargetNames.Num() == 0)
 		{
-			return true;
+			break; 
 		}
-		if(Visitor.bSearchSubDirectories)
+		if (Visitor.bSearchPluginForTargets)
 		{
-			DirectoryNames += Visitor.SubDirectories;
+			FString SourceDir = DirectoryNames[Idx] / TEXT("Source");
+			if (FM.DirectoryExists(*SourceDir))
+			{
+				DirectoryNames.Emplace(MoveTemp(SourceDir));
+			}
+			FString TestsDir = DirectoryNames[Idx] / TEXT("Tests");
+			if (FM.DirectoryExists(*TestsDir))
+			{
+				DirectoryNames.Emplace(MoveTemp(TestsDir));				
+			}
+		}
+		if (Visitor.bSearchSubDirectories)
+		{
+			DirectoryNames.Append(MoveTemp(Visitor.SubDirectories));
 		}
 	}
 
+	for (const FString& Target : RemainingTargetNames)
+	{
+		UE_LOG(LogDesktopPlatform, Log, TEXT("Failed to find target file for %s, cache out of date."), *Target);
+	}
 	// If we found all the previous target files
 	return RemainingTargetNames.Num() == 0;
 }
@@ -1245,10 +1284,18 @@ const TArray<FTargetInfo>& FDesktopPlatformBase::GetTargetsForProject(const FStr
 	{
 		// Read it in and check it's still valid
 		TArray<FTargetInfo> NewTargets;
-		TArray<FString> DirectoryNames = { ProjectSourceDir, ProjectDir / TEXT("Platforms"), ProjectDir / TEXT("Restricted"), ProjectDir / TEXT("Plugins") };
-		if(ReadTargetInfo(InfoFileName, NewTargets) && IsTargetInfoValid(NewTargets, DirectoryNames, StatData.ModificationTime))
+		if(ReadTargetInfo(InfoFileName, NewTargets))
 		{
-			return ProjectFileToTargets.Emplace(MoveTemp(NormalizedProjectFile), MoveTemp(NewTargets));
+			TArray<FString> DirectoryNames = { 
+				ProjectSourceDir,
+				ProjectDir / TEXT("Plugins"),
+				ProjectDir / TEXT("Platforms"), 
+				ProjectDir / TEXT("Restricted"),
+				};
+			if(IsTargetInfoValid(NewTargets, DirectoryNames, StatData.ModificationTime))
+			{
+				return ProjectFileToTargets.Emplace(MoveTemp(NormalizedProjectFile), MoveTemp(NewTargets));
+			}
 		}
 	}
 
@@ -1265,6 +1312,7 @@ const TArray<FTargetInfo>& FDesktopPlatformBase::GetTargetsForProject(const FStr
 		Arguments += FString::Printf(TEXT(" -Project=\"%s\""), *IFileManager::Get().ConvertToAbsolutePathForExternalAppForRead(*ProjectFile));
 	}
 	Arguments += FString::Printf(TEXT(" -Output=\"%s\""), *IFileManager::Get().ConvertToAbsolutePathForExternalAppForRead(*InfoFileName));
+	Arguments += TEXT(" -IncludeAllTargets");
 
 	// Run UBT to update the list of targets. Try to run it without building first.
 	FString Output;
@@ -1823,6 +1871,11 @@ bool FDesktopPlatformBase::ReadTargetInfo(const FString& FileName, TArray<FTarge
 		{
 			return false;
 		}
+		
+		if (bool bDefaultTarget; TargetObject.TryGetBoolField(TEXT("DefaultTarget"), bDefaultTarget))
+		{
+			Targets[Idx].DefaultTarget = bDefaultTarget;
+		}
 
 		FString Path = FPaths::ConvertRelativePathToFull(FPaths::Combine(BaseDir, Targets[Idx].Path));
 		FPaths::MakePlatformFilename(Path);
diff --git a/Engine/Source/Developer/DesktopPlatform/Public/IDesktopPlatform.h b/Engine/Source/Developer/DesktopPlatform/Public/IDesktopPlatform.h
--- a/Engine/Source/Developer/DesktopPlatform/Public/IDesktopPlatform.h
+++ b/Engine/Source/Developer/DesktopPlatform/Public/IDesktopPlatform.h
@@ -67,6 +67,7 @@ struct FTargetInfo
 	FString Name;
 	FString Path;
 	EBuildTargetType Type;
+	TOptional<bool> DefaultTarget;
 };
 
 /**
diff --git a/Engine/Source/Developer/LauncherServices/Private/Profiles/LauncherProfile.h b/Engine/Source/Developer/LauncherServices/Private/Profiles/LauncherProfile.h
--- a/Engine/Source/Developer/LauncherServices/Private/Profiles/LauncherProfile.h
+++ b/Engine/Source/Developer/LauncherServices/Private/Profiles/LauncherProfile.h
@@ -75,7 +75,7 @@ inline bool TryGetDefaultTargetName(const FString& ProjectFile, EBuildTargetType
 	const TArray<FTargetInfo>& Targets = FDesktopPlatformModule::Get()->GetTargetsForProject(ProjectFile);
 	for (const FTargetInfo& Target : Targets)
 	{
-		if (Target.Type == TargetType)
+		if (Target.Type == TargetType && Target.DefaultTarget.Get(true))
 		{
 			OutTargetName = Target.Name;
 			return true;
