PR #9894: Adding support setting SteamDeck password

#tests Compiled/Run locally

diff --git a/Engine/Source/Developer/DesktopPlatform/Public/SteamDeck/SteamDeckDevice.h b/Engine/Source/Developer/DesktopPlatform/Public/SteamDeck/SteamDeckDevice.h
--- a/Engine/Source/Developer/DesktopPlatform/Public/SteamDeck/SteamDeckDevice.h
+++ b/Engine/Source/Developer/DesktopPlatform/Public/SteamDeck/SteamDeckDevice.h
@@ -15,10 +15,11 @@ template<class ParentDeviceClass>
 class TSteamDeckDevice : public ParentDeviceClass
 {
 public:
-	TSteamDeckDevice(FString InIpAddr, FString InDeviceName, FString InUserName, const ITargetPlatform& InTargetPlatform, const TCHAR* InRuntimeOSName)
+	TSteamDeckDevice(FString InIpAddr, FString InDeviceName, FString InUserName, FString InPassword, const ITargetPlatform& InTargetPlatform, const TCHAR* InRuntimeOSName)
 		: ParentDeviceClass(InTargetPlatform)
 		, IpAddr(InIpAddr)
 		, UserName(InUserName)
+		, Password(InPassword)
 		, RuntimeOSName(InRuntimeOSName)
 	{
 		DeviceName = FString::Printf(TEXT("%s (%s)"), *InDeviceName, InRuntimeOSName);
@@ -42,9 +43,7 @@ public:
 	virtual bool GetUserCredentials(FString& OutUserName, FString& OutUserPassword) override
 	{
 		OutUserName = UserName;
-
-		// No need for a password, as it will use an rsa key that is part of the SteamOS Devkit Client
-		OutUserPassword = FString();
+		OutUserPassword = Password;
 
 		return true;
 	}
@@ -53,7 +52,7 @@ public:
 	{
 		TArray<FString> EngineIniSteamDeckDevices;
 
-		// Expected ini format: +SteamDeckDevice=(IpAddr=10.1.33.19,Name=MySteamDeck,UserName=deck)
+		// Expected ini format: +SteamDeckDevice=(IpAddr=10.1.33.19,Name=MySteamDeck,UserName=deck,Password=password)
 		GConfig->GetArray(TEXT("SteamDeck"), TEXT("SteamDeckDevice"), EngineIniSteamDeckDevices, GEngineIni);
 
 		TArray<ITargetDevicePtr> SteamDevices;
@@ -62,6 +61,7 @@ public:
 			FString IpAddr;
 			FString Name;
 			FString ConfigUserName;
+			FString ConfigPassword;
 			
 			// As of SteamOS version 3.1, "deck" is the required UserName to be used when making a remote connection
 			// to the device. Eventually it will be configurable, so we will allow users to set it via the config. 
@@ -86,7 +86,13 @@ public:
 				ConfigUserName = DefaultUserName;
 			}
 
-			SteamDevices.Add(MakeShareable(new TSteamDeckDevice<ParentDeviceClass>(IpAddr, Name, ConfigUserName, TargetPlatform, RuntimeOSName)));
+			// If the user has not specified a Password in the config, use an rsa key that is part of the SteamOS Devkit Client
+			if (!FParse::Value(*Device, TEXT("Password="), ConfigPassword))
+			{
+				ConfigPassword = FString();
+			}
+
+			SteamDevices.Add(MakeShareable(new TSteamDeckDevice<ParentDeviceClass>(IpAddr, Name, ConfigUserName, ConfigPassword, TargetPlatform, RuntimeOSName)));
 		}
 
 		return SteamDevices;
@@ -96,5 +102,6 @@ private:
 	FString IpAddr;
 	FString DeviceName;
 	FString UserName;
+	FString Password;
 	FString RuntimeOSName;
 };
