[Backout] - CL30069106
Original CL Desc
-----------------------------------------------------------------
Horde: Add a typed HTTP client for communicating with the Horde server. Removes boilerplate code for getting an access token, configuring request objects, and parsing responses. Synchronous for now.

diff --git a/Engine/Source/Developer/DesktopPlatform/Private/DesktopPlatformBase.cpp b/Engine/Source/Developer/DesktopPlatform/Private/DesktopPlatformBase.cpp
--- a/Engine/Source/Developer/DesktopPlatform/Private/DesktopPlatformBase.cpp
+++ b/Engine/Source/Developer/DesktopPlatform/Private/DesktopPlatformBase.cpp
@@ -820,10 +820,7 @@ bool FDesktopPlatformBase::GetOidcAccessToken(const FString& RootDir, const FStr
 	FString Arguments = TEXT(" ");
 	Arguments += FString::Printf(TEXT(" --Service=\"%s\""), *ProviderIdentifier);
 	Arguments += FString::Printf(TEXT(" --OutFile=\"%s\""), *ResultFilePath);
-	if (ProjectFileName.Len() > 0)
-	{
-		Arguments += FString::Printf(TEXT(" --project=\"%s\""), *IFileManager::Get().ConvertToAbsolutePathForExternalAppForRead(*FPaths::GetPath(*ProjectFileName)));
-	}
+	Arguments += FString::Printf(TEXT(" --project=\"%s\""),  *IFileManager::Get().ConvertToAbsolutePathForExternalAppForRead(*FPaths::ProjectDir()));
 	FString UnattendedArguments = Arguments;
 	UnattendedArguments += TEXT(" --Unattended=true");
 
diff --git a/Engine/Source/Developer/Horde/Horde.Build.cs b/Engine/Source/Developer/Horde/Horde.Build.cs
--- a/Engine/Source/Developer/Horde/Horde.Build.cs
+++ b/Engine/Source/Developer/Horde/Horde.Build.cs
@@ -7,6 +7,6 @@ public class Horde : ModuleRules
 {
 	public Horde(ReadOnlyTargetRules Target) : base(Target)
 	{
-		PrivateDependencyModuleNames.AddRange(new string[] { "Core", "HTTP", "Json", "DesktopPlatform" });
+		PrivateDependencyModuleNames.AddRange(new string[] { "Core" });
 	}
 }
diff --git a/Engine/Source/Developer/Horde/Private/HordeHttpClient.cpp b/Engine/Source/Developer/Horde/Private/HordeHttpClient.cpp
deleted file mode 100644
diff --git a/Engine/Source/Developer/Horde/Private/Server/ServerMessages.cpp b/Engine/Source/Developer/Horde/Private/Server/ServerMessages.cpp
deleted file mode 100644
diff --git a/Engine/Source/Developer/Horde/Public/HordeHttpClient.h b/Engine/Source/Developer/Horde/Public/HordeHttpClient.h
deleted file mode 100644
diff --git a/Engine/Source/Developer/Horde/Public/Server/ServerMessages.h b/Engine/Source/Developer/Horde/Public/Server/ServerMessages.h
deleted file mode 100644
