[DesktopPlatform]
* Added exit code to error message when failing to run OidcToken.exe

diff --git a/Engine/Source/Developer/DesktopPlatform/Private/DesktopPlatformBase.cpp b/Engine/Source/Developer/DesktopPlatform/Private/DesktopPlatformBase.cpp
--- a/Engine/Source/Developer/DesktopPlatform/Private/DesktopPlatformBase.cpp
+++ b/Engine/Source/Developer/DesktopPlatform/Private/DesktopPlatformBase.cpp
@@ -857,7 +857,7 @@ bool FDesktopPlatformBase::GetOidcAccessToken(const FString& RootDir, const FStr
 
 	if (!bRes)
 	{
-		UE_LOG(LogDesktopPlatform, Warning, TEXT("Failed to run OidcToken (project file is '%s', exe path is '%s')"), *ProjectFileName, *GetOidcTokenExecutableFilename(RootDir));
+		UE_LOG(LogDesktopPlatform, Warning, TEXT("Failed to run OidcToken (project file is '%s', exe path is '%s'). ExitCode: %i"), *ProjectFileName, *GetOidcTokenExecutableFilename(RootDir), ExitCode);
 		return false;
 	}
 	
