Fixed typo in message log.

diff --git a/Engine/Source/Developer/DesktopPlatform/Private/DesktopPlatformBase.cpp b/Engine/Source/Developer/DesktopPlatform/Private/DesktopPlatformBase.cpp
--- a/Engine/Source/Developer/DesktopPlatform/Private/DesktopPlatformBase.cpp
+++ b/Engine/Source/Developer/DesktopPlatform/Private/DesktopPlatformBase.cpp
@@ -1632,7 +1632,7 @@ bool FDesktopPlatformBase::EnumerateProjectsKnownByEngine(const FString &Identif
 		}
 	}
 
-	UE_LOG(LogDesktopPlatform, Log, TEXT("Searcing for projects in .uprojectdirs"));
+	UE_LOG(LogDesktopPlatform, Log, TEXT("Searching for projects in .uprojectdirs"));
 
 	// Find all the native projects, and either add or remove them from the list depending on whether we want native projects
 	const FUProjectDictionary &Dictionary = GetCachedProjectDictionary(RootDir);
