Horde: Add a setting for configuring the Horde server to the editor.

- Falls back to the UE_HORDE_URL environment variable or value configured in the registry from UGS if not set.
- Added support to OIDCToken for configuring using Horde settings using the --HordeUrl argument.
- Added overloads to HordeHttpClient to query auth settings from Horde itself.

diff --git a/Engine/Source/Developer/DesktopPlatform/Private/DesktopPlatformBase.cpp b/Engine/Source/Developer/DesktopPlatform/Private/DesktopPlatformBase.cpp
--- a/Engine/Source/Developer/DesktopPlatform/Private/DesktopPlatformBase.cpp
+++ b/Engine/Source/Developer/DesktopPlatform/Private/DesktopPlatformBase.cpp
@@ -810,54 +810,34 @@
 
 bool FDesktopPlatformBase::GetOidcAccessToken(const FString& RootDir, const FString& ProjectFileName, const FString& ProviderIdentifier, bool bUnattended, FFeedbackContext* Warn, FString& OutToken, FDateTime& OutTokenExpiresAt, bool& bOutWasInteractiveLogin)
 {
+	FString BaseArguments;
+	BaseArguments += FString::Printf(TEXT(" --Service=\"%s\""), *ProviderIdentifier);
+	if (ProjectFileName.Len() > 0)
+	{
+		BaseArguments += FString::Printf(TEXT(" --project=\"%s\""), *IFileManager::Get().ConvertToAbsolutePathForExternalAppForRead(*FPaths::GetPath(*ProjectFileName)));
+	}
+
+	return GetOidcAccessTokenInternal(RootDir, BaseArguments, bUnattended, Warn, OutToken, OutTokenExpiresAt, bOutWasInteractiveLogin);
+}
+
+bool FDesktopPlatformBase::GetOidcTokenStatus(const FString& RootDir, const FString& ProjectFileName, const FString& ProviderIdentifier, FFeedbackContext* Warn, int& OutStatus)
+{
 	IFileManager::Get().MakeDirectory(*FPaths::ProjectIntermediateDir(), /*bTree*/ true);
-	FString ResultFilePath = FPaths::CreateTempFilename(*FPaths::ProjectIntermediateDir(), TEXT("oidcToken.json"));
-	ON_SCOPE_EXIT
-	{
-		IFileManager::Get().Delete(*ResultFilePath, false /* RequireExists */, true /* EvenIfReadOnly */, true /* Quiet */);
-	};
+	FString ResultFilePath = FPaths::CreateTempFilename(*FPaths::ProjectIntermediateDir(), TEXT("oidcToken-status.json"));
 
 	FString Arguments = TEXT(" ");
 	Arguments += FString::Printf(TEXT(" --Service=\"%s\""), *ProviderIdentifier);
 	Arguments += FString::Printf(TEXT(" --OutFile=\"%s\""), *ResultFilePath);
-	if (ProjectFileName.Len() > 0)
-	{
-		Arguments += FString::Printf(TEXT(" --project=\"%s\""), *IFileManager::Get().ConvertToAbsolutePathForExternalAppForRead(*FPaths::GetPath(*ProjectFileName)));
-	}
-	FString UnattendedArguments = Arguments;
-	UnattendedArguments += TEXT(" --Unattended=true");
+	Arguments += FString::Printf(TEXT(" --project=\"%s\""),  *IFileManager::Get().ConvertToAbsolutePathForExternalAppForRead(*FPaths::ProjectDir()));
 
-	// first we attempt to fetch a token using cached offline tokens, thus setting unattended
 	bool bRes = true;
 	int32 ExitCode;
 	FString ProcessStdout;
-	bRes = InvokeOidcTokenToolSync(LOCTEXT("GetOidcAccessToken", "Fetching OIDC Access Token..."), RootDir, UnattendedArguments, Warn, ExitCode, ProcessStdout);
+	bRes = InvokeOidcTokenToolSync(LOCTEXT("GetOidcAccessTokenStatus", "Fetching OIDC Access Token Status..."), RootDir, Arguments, Warn, ExitCode, ProcessStdout);
 
-	bOutWasInteractiveLogin = false;
-
-	if (ExitCode == 10)
-	{
-		bRes = GetOidcAccessTokenInteractive(RootDir, Arguments, bUnattended, Warn, ExitCode);
-
-		bOutWasInteractiveLogin = true;
-
-		if (!bRes)
-		{
-			if (bUnattended)
-			{
-				UE_LOG(LogDesktopPlatform, Warning, TEXT("Unable to allocate an access token. Unattended set so unable to complete interactive login. Make sure you start the editor and login once or log in using UGS or using the UGS cli command 'login'. Provider used: '%s'. Ran OidcToken (project file is '%s', exe path is '%s')"), *ProviderIdentifier, *ProjectFileName, *GetOidcTokenExecutableFilename(RootDir));
-			}
-			else
-			{
-				UE_LOG(LogDesktopPlatform, Error, TEXT("Unable to allocate an access token. Interactive login failed, make sure you are assigned access and are able to login in the created browser window. Provider used: '%s'. Ran OidcToken (project file is '%s', exe path is '%s')"), *ProviderIdentifier, *ProjectFileName, *GetOidcTokenExecutableFilename(RootDir));
-			}
-			return false;
-		}
-	}
-
 	if (!bRes)
 	{
-		UE_LOG(LogDesktopPlatform, Warning, TEXT("Failed to run OidcToken (project file is '%s', exe path is '%s'). ExitCode: %i"), *ProjectFileName, *GetOidcTokenExecutableFilename(RootDir), ExitCode);
+		UE_LOG(LogDesktopPlatform, Warning, TEXT("Failed to run OidcToken to determine token status (project file is '%s', exe path is '%s')"), *ProjectFileName, *GetOidcTokenExecutableFilename(RootDir));
 		return false;
 	}
 	
@@ -870,13 +850,10 @@
 		TSharedRef< TJsonReader<> > Reader = TJsonReaderFactory<>::Create(TokenText);
 		if(FJsonSerializer::Deserialize(Reader, Object) && Object.IsValid())
 		{
-			FString Token;
-			FString ExpiresAt;
-			if(Object->TryGetStringField(TEXT("Token"), Token) && Object->TryGetStringField(TEXT("ExpiresAt"), ExpiresAt))
+			int Status;
+			if(Object->TryGetNumberField(TEXT("Status"), Status))
 			{
-				OutToken = Token;
+				OutStatus = Status;
 
-				FDateTime::ParseIso8601(*ExpiresAt, OutTokenExpiresAt);
-
 				// Remove the output file if its still around
 				IFileManager::Get().Delete(*ResultFilePath, true, false, true);
@@ -886,32 +863,129 @@
 		}
 	}
 
-	UE_LOG(LogDesktopPlatform, Warning, TEXT("Failed to run OidcToken (project file is '%s', exe path is '%s'). No result file found at '%s', closed with exit code: %d"), *ProjectFileName, *GetOidcTokenExecutableFilename(RootDir), *ResultFilePath, ExitCode);
-
 	// Remove the output file if its still around
 	IFileManager::Get().Delete(*ResultFilePath, true, false, true);
 	return false;
 }
 
+bool FDesktopPlatformBase::GetHordeUrl(FString& OutHordeUrl)
+{
+	struct FServerUrlInitializer
+	{
+		static FString Get(IDesktopPlatform* DesktopPlatform)
+		{
+			FString Url;
+			if (!GConfig->GetString(TEXT("/Script/UnrealEd.EditorSettings"), TEXT("HordeUrl"), Url, GEditorSettingsIni) || Url.IsEmpty())
+			{
+				DesktopPlatform->GetDefaultHordeUrl(Url);
+			}
+			return Url;
+		}
+	};
 
-bool FDesktopPlatformBase::GetOidcTokenStatus(const FString& RootDir, const FString& ProjectFileName, const FString& ProviderIdentifier, FFeedbackContext* Warn, int& OutStatus)
+	static FString CachedValue = FServerUrlInitializer::Get(this);
+
+	OutHordeUrl = CachedValue;
+	return !OutHordeUrl.IsEmpty();
+}
+
+bool FDesktopPlatformBase::GetHordeAccessToken(const FString& HordeUrl, bool bUnattended, FFeedbackContext* Warn, FString& OutToken, FDateTime& OutTokenExpiresAt, bool& bOutWasInteractiveLogin)
 {
+	FString Token = FPlatformMisc::GetEnvironmentVariable(TEXT("UE_HORDE_TOKEN"));
+	if (!Token.IsEmpty())
+	{
+		FString Url = FPlatformMisc::GetEnvironmentVariable(TEXT("UE_HORDE_URL"));
+		if (Url == HordeUrl)
+		{
+			OutToken = Token;
+			OutTokenExpiresAt = FDateTime();
+			return true;
+		}
+	}
+
+	FString BaseArguments = FString::Printf(TEXT(" --HordeUrl=\"%s\""), *HordeUrl);
+	return GetOidcAccessTokenInternal(FPaths::EngineDir(), BaseArguments, bUnattended, Warn, OutToken, OutTokenExpiresAt, bOutWasInteractiveLogin);
+}
+
+bool FDesktopPlatformBase::GetDefaultHordeUrl(FString& OutHordeUrl)
+{
+	struct FServerUrlInitializer
+	{
+		static FString Get()
+		{
+			FString Url = FPlatformMisc::GetEnvironmentVariable(TEXT("UE_HORDE_URL"));
+			if (!Url.IsEmpty())
+			{
+				return Url;
+			}
+
+#if PLATFORM_WINDOWS
+			if (FWindowsPlatformMisc::QueryRegKey(HKEY_CURRENT_USER, TEXT("SOFTWARE\\Epic Games\\Horde"), TEXT("Url"), Url) && !Url.IsEmpty())
+			{
+				return Url;
+			}
+
+			if (FWindowsPlatformMisc::QueryRegKey(HKEY_LOCAL_MACHINE, TEXT("SOFTWARE\\Epic Games\\Horde"), TEXT("Url"), Url) && !Url.IsEmpty())
+			{
+				return Url;
+			}
+#endif
+			return FString();
+		}
+	};
+
+	static FString CachedValue = FServerUrlInitializer::Get();
+
+	OutHordeUrl = CachedValue;
+	return !OutHordeUrl.IsEmpty();
+}
+
+bool FDesktopPlatformBase::GetOidcAccessTokenInternal(const FString& RootDir, const FString& BaseArguments, bool bUnattended, FFeedbackContext* Warn, FString& OutToken, FDateTime& OutTokenExpiresAt, bool& bOutWasInteractiveLogin)
+{
 	IFileManager::Get().MakeDirectory(*FPaths::ProjectIntermediateDir(), /*bTree*/ true);
-	FString ResultFilePath = FPaths::CreateTempFilename(*FPaths::ProjectIntermediateDir(), TEXT("oidcToken-status.json"));
+	FString ResultFilePath = FPaths::CreateTempFilename(*FPaths::ProjectIntermediateDir(), TEXT("oidcToken.json"));
+	ON_SCOPE_EXIT
+	{
+		IFileManager::Get().Delete(*ResultFilePath, false /* RequireExists */, true /* EvenIfReadOnly */, true /* Quiet */);
+	};
 
-	FString Arguments = TEXT(" ");
-	Arguments += FString::Printf(TEXT(" --Service=\"%s\""), *ProviderIdentifier);
+	FString Arguments = TEXT(" ") + BaseArguments;
 	Arguments += FString::Printf(TEXT(" --OutFile=\"%s\""), *ResultFilePath);
-	Arguments += FString::Printf(TEXT(" --project=\"%s\""),  *IFileManager::Get().ConvertToAbsolutePathForExternalAppForRead(*FPaths::ProjectDir()));
 
+	FString UnattendedArguments = Arguments;
+	UnattendedArguments += TEXT(" --Unattended=true");
+
+	// first we attempt to fetch a token using cached offline tokens, thus setting unattended
 	bool bRes = true;
 	int32 ExitCode;
 	FString ProcessStdout;
-	bRes = InvokeOidcTokenToolSync(LOCTEXT("GetOidcAccessTokenStatus", "Fetching OIDC Access Token Status..."), RootDir, Arguments, Warn, ExitCode, ProcessStdout);
+	bRes = InvokeOidcTokenToolSync(LOCTEXT("GetOidcAccessToken", "Fetching OIDC Access Token..."), RootDir, UnattendedArguments, Warn, ExitCode, ProcessStdout);
 
+	bOutWasInteractiveLogin = false;
+
+	if (ExitCode == 10)
+	{
+		bRes = GetOidcAccessTokenInteractive(RootDir, Arguments, bUnattended, Warn, ExitCode);
+
+		bOutWasInteractiveLogin = true;
+
+		if (!bRes)
+		{
+			if (bUnattended)
+			{
+				UE_LOG(LogDesktopPlatform, Warning, TEXT("Unable to allocate an access token. Unattended set so unable to complete interactive login. Make sure you start the editor and login once or log in using UGS or using the UGS cli command 'login'. (Ran '%s%s')"), *GetOidcTokenExecutableFilename(RootDir), *Arguments);
+			}
+			else
+			{
+				UE_LOG(LogDesktopPlatform, Error, TEXT("Unable to allocate an access token. Interactive login failed, make sure you are assigned access and are able to login in the created browser window. (Ran '%s%s')"), *GetOidcTokenExecutableFilename(RootDir), *Arguments);
+			}
+			return false;
+		}
+	}
+
 	if (!bRes)
 	{
-		UE_LOG(LogDesktopPlatform, Warning, TEXT("Failed to run OidcToken to determine token status (project file is '%s', exe path is '%s')"), *ProjectFileName, *GetOidcTokenExecutableFilename(RootDir));
+		UE_LOG(LogDesktopPlatform, Warning, TEXT("Failed to run OidcToken ('%s%s'). ExitCode: %i"), *GetOidcTokenExecutableFilename(RootDir), *Arguments, ExitCode);
 		return false;
 	}
 	
@@ -924,10 +998,13 @@
 		TSharedRef< TJsonReader<> > Reader = TJsonReaderFactory<>::Create(TokenText);
 		if(FJsonSerializer::Deserialize(Reader, Object) && Object.IsValid())
 		{
-			int Status;
-			if(Object->TryGetNumberField(TEXT("Status"), Status))
+			FString Token;
+			FString ExpiresAt;
+			if(Object->TryGetStringField(TEXT("Token"), Token) && Object->TryGetStringField(TEXT("ExpiresAt"), ExpiresAt))
 			{
-				OutStatus = Status;
+				OutToken = Token;
 
+				FDateTime::ParseIso8601(*ExpiresAt, OutTokenExpiresAt);
+
 				// Remove the output file if its still around
 				IFileManager::Get().Delete(*ResultFilePath, true, false, true);
@@ -937,6 +1014,8 @@
 		}
 	}
 
+	UE_LOG(LogDesktopPlatform, Warning, TEXT("Failed to run OidcToken ('%s%s'). No result file found at '%s', closed with exit code: %d"), *GetOidcTokenExecutableFilename(RootDir), *Arguments, *ResultFilePath, ExitCode);
+
 	// Remove the output file if its still around
 	IFileManager::Get().Delete(*ResultFilePath, true, false, true);
 	return false;
diff --git a/Engine/Source/Developer/DesktopPlatform/Private/DesktopPlatformBase.h b/Engine/Source/Developer/DesktopPlatform/Private/DesktopPlatformBase.h
--- a/Engine/Source/Developer/DesktopPlatform/Private/DesktopPlatformBase.h
+++ b/Engine/Source/Developer/DesktopPlatform/Private/DesktopPlatformBase.h
@@ -58,6 +58,10 @@ public:
     virtual bool GetOidcAccessToken(const FString& RootDir, const FString& ProjectFileName, const FString& ProviderIdentifier, bool bUnattended, FFeedbackContext* Warn, FString& OutToken, FDateTime& OutTokenExpiresAt, bool& bOutWasInteractiveLogin) override;
 	virtual bool GetOidcTokenStatus(const FString& RootDir, const FString& ProjectFileName, const FString& ProviderIdentifier, FFeedbackContext* Warn, int& OutStatus) override;
 
+	virtual bool GetHordeUrl(FString& OutHordeUrl) override;
+	virtual bool GetHordeAccessToken(const FString& HordeUrl, bool bUnattended, FFeedbackContext* Warn, FString& OutToken, FDateTime& OutTokenExpiresAt, bool& bOutWasInteractiveLogin) override;
+	virtual bool GetDefaultHordeUrl(FString& OutHordeUrl) override;
+
 	virtual const TArray<FTargetInfo>& GetTargetsForProject(const FString& ProjectFile) const override;
 	virtual const TArray<FTargetInfo>& GetTargetsForCurrentProject() const override;
 
@@ -85,6 +89,7 @@ private:
 
 	bool InvokeOidcTokenToolSync(const FText& Description, const FString& RootDir, const FString& Arguments, FFeedbackContext* Warn, int32& OutReturnCode, FString& OutProcOutput);
 	FProcHandle InvokeOidcTokenToolAsync(const FString& InArguments, void*& OutReadPipe, void*& OutWritePipe);
+	bool GetOidcAccessTokenInternal(const FString& RootDir, const FString& Arguments, bool bUnattended, FFeedbackContext* Warn, FString& OutToken, FDateTime& OutTokenExpiresAt, bool& bOutWasInteractiveLogin);
 	bool GetOidcAccessTokenInteractive(const FString& RootDir, const FString& Arguments, bool bUnattended, FFeedbackContext* Warn, int32& OutReturnCode);
 
 	static bool ReadTargetInfo(const FString& FileName, TArray<FTargetInfo>& Targets);
diff --git a/Engine/Source/Developer/DesktopPlatform/Public/IDesktopPlatform.h b/Engine/Source/Developer/DesktopPlatform/Public/IDesktopPlatform.h
--- a/Engine/Source/Developer/DesktopPlatform/Public/IDesktopPlatform.h
+++ b/Engine/Source/Developer/DesktopPlatform/Public/IDesktopPlatform.h
@@ -484,4 +484,33 @@ public:
 	* @return true if the task completed successfully.
 	*/
 	virtual bool GetOidcTokenStatus(const FString& RootDir, const FString& ProjectFileName, const FString& ProviderIdentifier, FFeedbackContext* Warn, int& OutStatus) = 0;
+
+	/**
+	* Get the URL of the configured Horde server.
+	*
+	* @param OutHordeUrl        URL of the Horde server
+	* @return true if a valid URL is returned.
+	*/
+	virtual bool GetHordeUrl(FString& OutHordeUrl) = 0;
+
+	/**
+	* Gets an access token for the given Horde Server.
+	*
+	* @param HordeUrl           URL of the Horde server.
+	* @param bUnattended		True to indicate that no user interaction should be assumed
+	* @param Warn				Feedback context to use for progress updates
+	* @param OutToken			The allocated access token
+	* @param OutTokenExpiresAt	When the token expires
+	* @param bOutWasInteractiveLogin	True if the interactive login flow was used
+	* @return true if the task completed successfully.
+	*/
+	virtual bool GetHordeAccessToken(const FString& HordeUrl, bool bUnattended, FFeedbackContext* Warn, FString& OutToken, FDateTime& OutTokenExpiresAt, bool& bOutWasInteractiveLogin) = 0;
+
+	/**
+	* Get the URL of the configured Horde server for the system, ignoring the default editor settings.
+	*
+	* @param OutHordeUrl        URL of the Horde server
+	* @return true if a valid URL is returned.
+	*/
+	virtual bool GetDefaultHordeUrl(FString& OutHordeUrl) = 0;
 };
diff --git a/Engine/Source/Developer/Horde/Private/Horde.cpp b/Engine/Source/Developer/Horde/Private/Horde.cpp
--- a/Engine/Source/Developer/Horde/Private/Horde.cpp
+++ b/Engine/Source/Developer/Horde/Private/Horde.cpp
@@ -5,6 +5,8 @@
 #include "Misc/Parse.h"
 #include "Modules/ModuleManager.h"
 #include "Modules/ModuleInterface.h"
+#include "IDesktopPlatform.h"
+#include "DesktopPlatformModule.h"
 
 struct FHordeModule : IModuleInterface
 {
@@ -17,6 +19,16 @@ IMPLEMENT_MODULE(FHordeModule, Horde)
 
 // --------------------------------------------------------------------------------
 
+bool FHorde::GetServerUrl(FString& OutUrl)
+{
+	IDesktopPlatform* DesktopPlatform = FDesktopPlatformModule::TryGet();
+	if (DesktopPlatform == nullptr)
+	{
+		return false;
+	}
+	return DesktopPlatform->GetHordeUrl(OutUrl);
+}
+
 FString FHorde::GetTemplateName()
 {
 	static FString TemplateName;
diff --git a/Engine/Source/Developer/Horde/Private/HordeHttpClient.cpp b/Engine/Source/Developer/Horde/Private/HordeHttpClient.cpp
--- a/Engine/Source/Developer/Horde/Private/HordeHttpClient.cpp
+++ b/Engine/Source/Developer/Horde/Private/HordeHttpClient.cpp
@@ -18,6 +18,24 @@ FHordeHttpClient::~FHordeHttpClient()
 {
 }
 
+bool FHordeHttpClient::Login(bool bUnattended, FFeedbackContext* Warn)
+{
+	IDesktopPlatform* DesktopPlatform = FDesktopPlatformModule::TryGet();
+	if (DesktopPlatform != nullptr)
+	{
+		FString NewToken;
+		FDateTime ExpiresAt;
+		bool bWasInteractive = false;
+
+		if (DesktopPlatform->GetHordeAccessToken(ServerUrl, bUnattended, Warn, NewToken, ExpiresAt, bWasInteractive))
+		{
+			Token = NewToken;
+			return true;
+		}
+	}
+	return false;
+}
+
 bool FHordeHttpClient::LoginWithOidc(const TCHAR* Profile, bool bUnattended, FFeedbackContext* Warn)
 {
 	IDesktopPlatform* DesktopPlatform = FDesktopPlatformModule::TryGet();
diff --git a/Engine/Source/Developer/Horde/Public/Horde.h b/Engine/Source/Developer/Horde/Public/Horde.h
--- a/Engine/Source/Developer/Horde/Public/Horde.h
+++ b/Engine/Source/Developer/Horde/Public/Horde.h
@@ -10,6 +10,10 @@
 
 struct FHorde
 {
+	// Gets the server URL. Equivalent to calling IDesktopPlatform::GetHordeUrl().
+	HORDE_API static bool GetServerUrl(FString& OutUrl);
+
+	// Settings available from the environment when running under a job
 	HORDE_API static FString GetTemplateName();
 	HORDE_API static FString GetTemplateId();
 	HORDE_API static FString GetServerURL();
diff --git a/Engine/Source/Developer/Horde/Public/HordeHttpClient.h b/Engine/Source/Developer/Horde/Public/HordeHttpClient.h
--- a/Engine/Source/Developer/Horde/Public/HordeHttpClient.h
+++ b/Engine/Source/Developer/Horde/Public/HordeHttpClient.h
@@ -13,6 +13,8 @@ public:
 	FHordeHttpClient(FString InServerUrl);
 	~FHordeHttpClient();
 
+	bool Login(bool bUnattended, FFeedbackContext* Warn = nullptr);
+
 	bool LoginWithOidc(const TCHAR* Profile, bool bUnattended, FFeedbackContext* Warn = nullptr);
 
 	bool LoginWithEnvironmentVariable();
diff --git a/Engine/Source/Editor/UnrealEd/Public/Settings/EditorSettings.h b/Engine/Source/Editor/UnrealEd/Public/Settings/EditorSettings.h
--- a/Engine/Source/Editor/UnrealEd/Public/Settings/EditorSettings.h
+++ b/Engine/Source/Editor/UnrealEd/Public/Settings/EditorSettings.h
@@ -115,6 +115,13 @@ class UEditorSettings : public UObject
 	UPROPERTY(EditAnywhere, Category="Derived Data Cache S3", meta = (DisplayName = "Global Local S3DDC Path", ConfigRestartRequired = true, EditCondition = "bEnableS3DDC"))
 	FDirectoryPath GlobalS3DDCPath;
 
+	// Horde Settings
+	// =====================================================================
+
+	/** URL of the Horde server */
+	UPROPERTY(EditAnywhere, config, Category = "Horde", meta = (DisplayName = "Horde Server URL", ConfigRestartRequired = true))
+	FString HordeUrl;
+
 	// =====================================================================
 
 	/** When checked, the most recently loaded project will be auto-loaded at editor startup if no other project was specified on the command line */
diff --git a/Engine/Source/Programs/OidcToken/Program.cs b/Engine/Source/Programs/OidcToken/Program.cs
--- a/Engine/Source/Programs/OidcToken/Program.cs
+++ b/Engine/Source/Programs/OidcToken/Program.cs
@@ -1,14 +1,19 @@
 // Copyright Epic Games, Inc. All Rights Reserved.
 
+using EpicGames.OIDC;
 using Microsoft.Extensions.Configuration;
+using Microsoft.Extensions.Configuration.Memory;
 using Microsoft.Extensions.DependencyInjection;
 using Microsoft.Extensions.Hosting;
+using Microsoft.Extensions.Logging;
 using System;
+using System.Collections.Generic;
 using System.IO;
-using System.Threading.Tasks;
-using Microsoft.Extensions.Logging;
 using System.Linq;
-using EpicGames.OIDC;
+using System.Net.Http;
+using System.Net.Http.Json;
+using System.Text.Json;
+using System.Threading.Tasks;
 
 namespace OidcToken
 {
@@ -19,10 +24,11 @@ namespace OidcToken
 			if (args.Any(s => s.Equals("--help") || s.Equals("-help")) || args.Length == 0)
 			{
 				// print help
-				Console.WriteLine("Usage: OidcToken --Service <serviceName> [options]");
-				Console.WriteLine("Service is a required parameter to indicate which OIDC service you intend to connect to. The connection details of the service is configured in appsettings.json/oidc-configuration.json");
+				Console.WriteLine("Usage: OidcToken [options]");
 				Console.WriteLine();
 				Console.WriteLine("Options: ");
+				Console.WriteLine(" --Service <serviceName> - Indicate which OIDC service you intend to connect to. The connection details of the service is configured in appsettings.json/oidc-configuration.json.");
+				Console.WriteLine(" --HordeUrl <url> - Specifies the URL of a Horde server to read configuration from.");
 				Console.WriteLine(" --Mode [Query/GetToken] - Switch mode to allow you to preview operation without triggering user interaction (result can be used to determine if user interaction is required)");
 				Console.WriteLine(" --OutFile <path> - Path to create json file of result");
 				Console.WriteLine(" --ResultToConsole [true/false] - If true the resulting json file is output to stdout (and logs are not created)");
@@ -43,23 +49,23 @@ namespace OidcToken
 			TokenServiceOptions options = new();
 			config.Bind(options);
 
-			// guess where the engine directory is based on the assumption that we are running out of Engine\Binaries\DotNET\OidcToken\<platform>
-			DirectoryInfo engineDir = new DirectoryInfo(Path.Combine(AppContext.BaseDirectory, "../../../../../Engine"));
-			if (!engineDir.Exists)
+			GetHordeAuthConfigResponse? hordeAuthConfig = null;
+			if (options.HordeUrl != null)
 			{
-				// try to see if engine dir can be found from the current code path Engine\Source\Programs\OidcToken\bin\<Configuration>\<.net-version>
-				engineDir = new DirectoryInfo(Path.Combine(AppContext.BaseDirectory, "../../../../../../../Engine"));
-
-				if (!engineDir.Exists)
-				{
-					throw new Exception($"Unable to guess engine directory so unable to continue running. Starting directory was: {AppContext.BaseDirectory}");
-				}
+				hordeAuthConfig = ReadHordeConfigurationAsync(options.HordeUrl).Result;
 			}
 
 			await Host.CreateDefaultBuilder(args)
 				.ConfigureAppConfiguration(builder =>
 				{
 					builder.AddConfiguration(config);
+
+					if (hordeAuthConfig != null && !String.IsNullOrEmpty(hordeAuthConfig.ProfileName))
+					{
+						Dictionary<string, string?> values = new Dictionary<string, string?>();
+						values[nameof(TokenServiceOptions.Service)] = hordeAuthConfig.ProfileName;
+						builder.AddInMemoryCollection(values);
+					}
 				})
 				.ConfigureLogging(loggingBuilder =>
 				{
@@ -75,7 +81,35 @@ namespace OidcToken
 				{
 					IConfiguration configuration = content.Configuration;
 					services.AddOptions<TokenServiceOptions>().Bind(configuration).ValidateDataAnnotations();
-					services.AddOptions<OidcTokenOptions>().Bind(ProviderConfigurationFactory.ReadConfiguration(engineDir, !string.IsNullOrEmpty(options.Project) ? new DirectoryInfo(options.Project) : null)).ValidateDataAnnotations();
+
+					IConfiguration serviceConfig;
+					if (hordeAuthConfig != null)
+					{
+						Dictionary<string, string?> values = new Dictionary<string, string?>();
+						values[$"{nameof(OidcTokenOptions.Providers)}:{hordeAuthConfig.ProfileName}:{nameof(ProviderInfo.DisplayName)}"] = hordeAuthConfig.ProfileName;
+						values[$"{nameof(OidcTokenOptions.Providers)}:{hordeAuthConfig.ProfileName}:{nameof(ProviderInfo.ServerUri)}"] = hordeAuthConfig.ServerUrl;
+						values[$"{nameof(OidcTokenOptions.Providers)}:{hordeAuthConfig.ProfileName}:{nameof(ProviderInfo.ClientId)}"] = hordeAuthConfig.ClientId;
+						values[$"{nameof(OidcTokenOptions.Providers)}:{hordeAuthConfig.ProfileName}:{nameof(ProviderInfo.RedirectUri)}"] = hordeAuthConfig.LocalRedirectUrls![0];
+						serviceConfig = new ConfigurationBuilder().AddInMemoryCollection(values).Build();
+					}
+					else
+					{
+						// guess where the engine directory is based on the assumption that we are running out of Engine\Binaries\DotNET\OidcToken\<platform>
+						DirectoryInfo engineDir = new DirectoryInfo(Path.Combine(AppContext.BaseDirectory, "../../../../../Engine"));
+						if (!engineDir.Exists)
+						{
+							// try to see if engine dir can be found from the current code path Engine\Source\Programs\OidcToken\bin\<Configuration>\<.net-version>
+							engineDir = new DirectoryInfo(Path.Combine(AppContext.BaseDirectory, "../../../../../../../Engine"));
+
+							if (!engineDir.Exists)
+							{
+								throw new Exception($"Unable to guess engine directory so unable to continue running. Starting directory was: {AppContext.BaseDirectory}");
+							}
+						}
+
+						serviceConfig = ProviderConfigurationFactory.ReadConfiguration(engineDir, !string.IsNullOrEmpty(options.Project) ? new DirectoryInfo(options.Project) : null);
+					}
+					services.AddOptions<OidcTokenOptions>().Bind(serviceConfig).ValidateDataAnnotations();
 
 					services.AddSingleton<OidcTokenManager>();
 					services.AddTransient<ITokenStore>(TokenStoreFactory.CreateTokenStore);
@@ -84,5 +118,45 @@ namespace OidcToken
 				})
 				.RunConsoleAsync();
 		}
+
+		class GetHordeAuthConfigResponse
+		{
+			public string ProfileName { get; set; } = null!;
+			public string ServerUrl { get; set; } = null!;
+			public string ClientId { get; set; } = null!;
+			public List<string> LocalRedirectUrls { get; set; } = new List<string>();
+		}
+
+		static async Task<GetHordeAuthConfigResponse> ReadHordeConfigurationAsync(Uri hordeUrl)
+		{
+			// Read the configuration settings from the Horde server
+			GetHordeAuthConfigResponse? authConfig;
+			using (HttpClient httpClient = new HttpClient())
+			{
+				using HttpRequestMessage request = new HttpRequestMessage(HttpMethod.Get, new Uri(hordeUrl, "api/v1/server/auth"));
+				using HttpResponseMessage response = await httpClient.SendAsync(request);
+				response.EnsureSuccessStatusCode();
+
+				JsonSerializerOptions jsonOptions = new JsonSerializerOptions { PropertyNameCaseInsensitive = true };
+				authConfig = await response.Content.ReadFromJsonAsync<GetHordeAuthConfigResponse>(jsonOptions);
+
+				if (authConfig == null)
+				{
+					throw new InvalidDataException("Server returned an empty auth config object");
+				}
+			}
+
+			string? localRedirectUrl = authConfig.LocalRedirectUrls.FirstOrDefault();
+			if (String.IsNullOrEmpty(authConfig.ServerUrl) || String.IsNullOrEmpty(authConfig.ClientId) || String.IsNullOrEmpty(localRedirectUrl))
+			{
+				throw new Exception("No auth server configuration found");
+			}
+
+			if (String.IsNullOrEmpty(authConfig.ProfileName))
+			{
+				authConfig.ProfileName = hordeUrl.Host.ToString();
+			}
+			return authConfig;
+		}
 	}
 }
\ No newline at end of file
diff --git a/Engine/Source/Programs/OidcToken/TokenService.cs b/Engine/Source/Programs/OidcToken/TokenService.cs
--- a/Engine/Source/Programs/OidcToken/TokenService.cs
+++ b/Engine/Source/Programs/OidcToken/TokenService.cs
@@ -117,7 +117,7 @@ namespace OidcToken
 						}
 					}
 
-					if (!tokenInfo.IsValid)
+					if (!tokenInfo.IsValid(DateTimeOffset.Now))
 					{
 						throw new Exception("Failed to allocate a token");
 					}
@@ -265,6 +265,11 @@ namespace OidcToken
 		/// </summary>
 		[Required] public string Service { get; set; } = null!;
 
+		/// <summary>
+		/// URL of the horde server to read provider information from
+		/// </summary>
+		public Uri? HordeUrl { get; set; }
+
 		/// <summary>
 		/// The mode we are running OidcToken in
 		/// </summary>
