Horde: Fix URL comparison for using access token from the environment.

diff --git a/Engine/Source/Developer/DesktopPlatform/Private/DesktopPlatformBase.cpp b/Engine/Source/Developer/DesktopPlatform/Private/DesktopPlatformBase.cpp
--- a/Engine/Source/Developer/DesktopPlatform/Private/DesktopPlatformBase.cpp
+++ b/Engine/Source/Developer/DesktopPlatform/Private/DesktopPlatformBase.cpp
@@ -895,7 +895,7 @@ bool FDesktopPlatformBase::GetHordeAccessToken(const FString& HordeUrl, bool bUn
 	if (!Token.IsEmpty())
 	{
 		FString Url = FPlatformMisc::GetEnvironmentVariable(TEXT("UE_HORDE_URL"));
-		if (Url == HordeUrl)
+		if (Url / TEXT("") == HordeUrl / TEXT(""))
 		{
 			OutToken = Token;
 			OutTokenExpiresAt = FDateTime();
