Horde: Unify the locations probed for the Horde server URL setting between C# tools (Horde Agent, UBT, UGS, etc...) and C++ (Editor).

* On Windows, this setting is stored in the registry under HKCU/SOFTWARE/Epic Games/Horde/Url.
* On other platforms, this setting is stored in the $HOME/.horde.json file.

The setting is no longer persisted to the editor config file.

diff --git a/Engine/Source/Developer/DesktopPlatform/Private/DesktopPlatformBase.cpp b/Engine/Source/Developer/DesktopPlatform/Private/DesktopPlatformBase.cpp
--- a/Engine/Source/Developer/DesktopPlatform/Private/DesktopPlatformBase.cpp
+++ b/Engine/Source/Developer/DesktopPlatform/Private/DesktopPlatformBase.cpp
@@ -868,27 +868,99 @@ bool FDesktopPlatformBase::GetOidcTokenStatus(const FString& RootDir, const FStr
 	return false;
 }
 
-bool FDesktopPlatformBase::GetHordeUrl(FString& OutHordeUrl)
+FString FDesktopPlatformBase::ReadHordeUrlWithoutCache()
 {
-	struct FServerUrlInitializer
+	FString Url = FPlatformMisc::GetEnvironmentVariable(TEXT("UE_HORDE_URL"));
+	if (!Url.IsEmpty())
+	{
+		return Url;
+	}
+
+#if PLATFORM_WINDOWS
+	if (FWindowsPlatformMisc::QueryRegKey(HKEY_CURRENT_USER, TEXT("SOFTWARE\\Epic Games\\Horde"), TEXT("Url"), Url) && !Url.IsEmpty())
 	{
-		static FString Get(IDesktopPlatform* DesktopPlatform)
+		return Url;
+	}
+
+	if (FWindowsPlatformMisc::QueryRegKey(HKEY_LOCAL_MACHINE, TEXT("SOFTWARE\\Epic Games\\Horde"), TEXT("Url"), Url) && !Url.IsEmpty())
+	{
+		return Url;
+	}
+#else
+	FString FileName = FPaths::Combine(FPlatformProcess::UserHomeDir(), TEXT(".horde.json"));
+
+	FString FileContents;
+	if (FFileHelper::LoadFileToString(FileContents, *FileName))
+	{
+		TSharedRef<TJsonReader<>> Reader = TJsonReaderFactory<>::Create(FileContents);
+
+		TSharedPtr<FJsonObject> Object;
+		if (FJsonSerializer::Deserialize(Reader, Object) && Object.IsValid() && Object->TryGetStringField(TEXT("Url"), Url))
 		{
-			FString Url;
-			if (!GConfig->GetString(TEXT("/Script/UnrealEd.EditorSettings"), TEXT("HordeUrl"), Url, GEditorSettingsIni) || Url.IsEmpty())
-			{
-				DesktopPlatform->GetDefaultHordeUrl(Url);
-			}
 			return Url;
 		}
-	};
+	}
+#endif
 
-	static FString CachedValue = FServerUrlInitializer::Get(this);
+	return FString();
+}
 
-	OutHordeUrl = CachedValue;
+bool FDesktopPlatformBase::GetHordeUrl(FString& OutHordeUrl)
+{
+	if (!bInitializedHordeServerUrl)
+	{
+		HordeServerUrl = ReadHordeUrlWithoutCache();
+		bInitializedHordeServerUrl = true;
+	}
+
+	OutHordeUrl = HordeServerUrl;
 	return !OutHordeUrl.IsEmpty();
 }
 
+void FDesktopPlatformBase::SetHordeUrl(const FString& HordeUrl)
+{
+	HordeServerUrl = HordeUrl;
+	bInitializedHordeServerUrl = true;
+
+#if PLATFORM_WINDOWS
+	if (HordeUrl.IsEmpty())
+	{
+		FPlatformMisc::DeleteStoredValue(TEXT("Epic Games"), TEXT("Horde"), TEXT("Url"));
+	}
+	else
+	{
+		FPlatformMisc::SetStoredValue(TEXT("Epic Games"), TEXT("Horde"), TEXT("Url"), HordeUrl);
+	}
+#else
+	FString FileName = FPaths::Combine(FPlatformProcess::UserHomeDir(), TEXT(".horde.json"));
+
+	// Read the existing json file
+	TSharedPtr<FJsonObject> Object;
+
+	FString FileContents;
+	if (FFileHelper::LoadFileToString(FileContents, *FileName))
+	{
+		TSharedRef<TJsonReader<>> Reader = TJsonReaderFactory<>::Create(FileContents);
+		FJsonSerializer::Deserialize(Reader, Object);
+	}
+
+	// Update the object with the new URL
+	if (!Object.IsValid())
+	{
+		Object = MakeShared<FJsonObject>();
+	}
+	Object->SetStringField(TEXT("Url"), HordeUrl);
+
+	// Write it back out agai
+	FString OutputFileContents;
+	TSharedRef<TJsonWriter<>> Writer = TJsonWriterFactory<>::Create(&OutputFileContents);
+	if (FJsonSerializer::Serialize(Object.ToSharedRef(), *Writer))
+	{
+		FFileHelper::SaveStringToFile(OutputFileContents, *FileName);
+	}
+#endif
+}
+
 bool FDesktopPlatformBase::GetHordeAccessToken(const FString& HordeUrl, bool bUnattended, FFeedbackContext* Warn, FString& OutToken, FDateTime& OutTokenExpiresAt, bool& bOutWasInteractiveLogin)
 {
 	FString Token = FPlatformMisc::GetEnvironmentVariable(TEXT("UE_HORDE_TOKEN"));
@@ -907,39 +979,6 @@ bool FDesktopPlatformBase::GetHordeAccessToken(const FString& HordeUrl, bool bUn
 	return GetOidcAccessTokenInternal(FPaths::EngineDir(), BaseArguments, bUnattended, Warn, OutToken, OutTokenExpiresAt, bOutWasInteractiveLogin);
 }
 
-bool FDesktopPlatformBase::GetDefaultHordeUrl(FString& OutHordeUrl)
-{
-	struct FServerUrlInitializer
-	{
-		static FString Get()
-		{
-			FString Url = FPlatformMisc::GetEnvironmentVariable(TEXT("UE_HORDE_URL"));
-			if (!Url.IsEmpty())
-			{
-				return Url;
-			}
-
-#if PLATFORM_WINDOWS
-			if (FWindowsPlatformMisc::QueryRegKey(HKEY_CURRENT_USER, TEXT("SOFTWARE\\Epic Games\\Horde"), TEXT("Url"), Url) && !Url.IsEmpty())
-			{
-				return Url;
-			}
-
-			if (FWindowsPlatformMisc::QueryRegKey(HKEY_LOCAL_MACHINE, TEXT("SOFTWARE\\Epic Games\\Horde"), TEXT("Url"), Url) && !Url.IsEmpty())
-			{
-				return Url;
-			}
-#endif
-			return FString();
-		}
-	};
-
-	static FString CachedValue = FServerUrlInitializer::Get();
-
-	OutHordeUrl = CachedValue;
-	return !OutHordeUrl.IsEmpty();
-}
-
 bool FDesktopPlatformBase::GetOidcAccessTokenInternal(const FString& RootDir, const FString& BaseArguments, bool bUnattended, FFeedbackContext* Warn, FString& OutToken, FDateTime& OutTokenExpiresAt, bool& bOutWasInteractiveLogin)
 {
 	IFileManager::Get().MakeDirectory(*FPaths::ProjectIntermediateDir(), /*bTree*/ true);
diff --git a/Engine/Source/Developer/DesktopPlatform/Private/DesktopPlatformBase.h b/Engine/Source/Developer/DesktopPlatform/Private/DesktopPlatformBase.h
--- a/Engine/Source/Developer/DesktopPlatform/Private/DesktopPlatformBase.h
+++ b/Engine/Source/Developer/DesktopPlatform/Private/DesktopPlatformBase.h
@@ -59,8 +59,8 @@ public:
 	virtual bool GetOidcTokenStatus(const FString& RootDir, const FString& ProjectFileName, const FString& ProviderIdentifier, FFeedbackContext* Warn, int& OutStatus) override;
 
 	virtual bool GetHordeUrl(FString& OutHordeUrl) override;
+	virtual void SetHordeUrl(const FString& HordeUrl) override;
 	virtual bool GetHordeAccessToken(const FString& HordeUrl, bool bUnattended, FFeedbackContext* Warn, FString& OutToken, FDateTime& OutTokenExpiresAt, bool& bOutWasInteractiveLogin) override;
-	virtual bool GetDefaultHordeUrl(FString& OutHordeUrl) override;
 
 	virtual const TArray<FTargetInfo>& GetTargetsForProject(const FString& ProjectFile) const override;
 	virtual const TArray<FTargetInfo>& GetTargetsForCurrentProject() const override;
@@ -74,6 +74,10 @@ private:
 	TMap<FString, FString> LauncherInstallationList;
 	TMap<FString, FUProjectDictionary> CachedProjectDictionaries;
 	mutable TMap<FString, TArray<FTargetInfo>> ProjectFileToTargets;
+	bool bInitializedHordeServerUrl;
+	FString HordeServerUrl;
+
+	static FString ReadHordeUrlWithoutCache();
 
 	void ReadLauncherInstallationList();
 	void CheckForLauncherEngineInstallation(const FString &AppId, const FString &Identifier, TMap<FString, FString> &OutInstallations);
diff --git a/Engine/Source/Developer/DesktopPlatform/Public/IDesktopPlatform.h b/Engine/Source/Developer/DesktopPlatform/Public/IDesktopPlatform.h
--- a/Engine/Source/Developer/DesktopPlatform/Public/IDesktopPlatform.h
+++ b/Engine/Source/Developer/DesktopPlatform/Public/IDesktopPlatform.h
@@ -486,13 +486,20 @@ public:
 	virtual bool GetOidcTokenStatus(const FString& RootDir, const FString& ProjectFileName, const FString& ProviderIdentifier, FFeedbackContext* Warn, int& OutStatus) = 0;
 
 	/**
-	* Get the URL of the configured Horde server.
+	* Get the URL of the configured Horde server for the current user account.
 	*
 	* @param OutHordeUrl        URL of the Horde server
 	* @return true if a valid URL is returned.
 	*/
 	virtual bool GetHordeUrl(FString& OutHordeUrl) = 0;
 
+	/**
+	* Set the URL of the configured Horde server for the current user account.
+	*
+	* @param HordeUrl        URL of the Horde server
+	*/
+	virtual void SetHordeUrl(const FString& HordeUrl) = 0;
+
 	/**
 	* Gets an access token for the given Horde Server.
 	*
@@ -505,12 +512,4 @@ public:
 	* @return true if the task completed successfully.
 	*/
 	virtual bool GetHordeAccessToken(const FString& HordeUrl, bool bUnattended, FFeedbackContext* Warn, FString& OutToken, FDateTime& OutTokenExpiresAt, bool& bOutWasInteractiveLogin) = 0;
-
-	/**
-	* Get the URL of the configured Horde server for the system, ignoring the default editor settings.
-	*
-	* @param OutHordeUrl        URL of the Horde server
-	* @return true if a valid URL is returned.
-	*/
-	virtual bool GetDefaultHordeUrl(FString& OutHordeUrl) = 0;
 };
diff --git a/Engine/Source/Editor/UnrealEd/Private/Settings/EditorSettings.cpp b/Engine/Source/Editor/UnrealEd/Private/Settings/EditorSettings.cpp
--- a/Engine/Source/Editor/UnrealEd/Private/Settings/EditorSettings.cpp
+++ b/Engine/Source/Editor/UnrealEd/Private/Settings/EditorSettings.cpp
@@ -16,6 +16,8 @@
 #include "UObject/NameTypes.h"
 #include "UObject/UnrealNames.h"
 #include "UObject/UnrealType.h"
+#include "IDesktopPlatform.h"
+#include "DesktopPlatformModule.h"
 
 UEditorSettings::UEditorSettings(const FObjectInitializer& ObjectInitializer)
 	: Super(ObjectInitializer)
@@ -32,6 +34,12 @@ UEditorSettings::UEditorSettings(const FObjectInitializer& ObjectInitializer)
 	FPlatformMisc::GetStoredValue(TEXT("Epic Games"), TEXT("GlobalDataCachePath"), TEXT("UE-LocalDataCachePath"), GlobalLocalDDCPath.Path);
 	FPlatformMisc::GetStoredValue(TEXT("Epic Games"), TEXT("GlobalDataCachePath"), TEXT("UE-SharedDataCachePath"), GlobalSharedDDCPath.Path);
 	FPlatformMisc::GetStoredValue(TEXT("Epic Games"), TEXT("GlobalDataCachePath"), TEXT("UE-S3DataCachePath"), GlobalS3DDCPath.Path);
+
+	IDesktopPlatform* DesktopPlatform = FDesktopPlatformModule::TryGet();
+	if (DesktopPlatform != nullptr)
+	{
+		DesktopPlatform->GetHordeUrl(HordeUrl);
+	}
 }
 
 bool UEditorSettings::CanEditChange(const FProperty* InProperty) const
@@ -111,6 +119,14 @@ void UEditorSettings::PostEditChangeProperty( struct FPropertyChangedEvent& Prop
 			FPlatformMisc::SetStoredValue(TEXT("Epic Games"), TEXT("GlobalDataCachePath"), TEXT("UE-S3DataCachePath"), GlobalS3DDCPath.Path);
 		}
 	}
+	else if (PropertyName == GET_MEMBER_NAME_CHECKED(UEditorSettings, HordeUrl))
+	{
+		IDesktopPlatform* DesktopPlatform = FDesktopPlatformModule::TryGet();
+		if (DesktopPlatform != nullptr)
+		{
+			DesktopPlatform->SetHordeUrl(HordeUrl);
+		}
+	}
 
 	SaveConfig(CPF_Config);
 }
diff --git a/Engine/Source/Editor/UnrealEd/Public/Settings/EditorSettings.h b/Engine/Source/Editor/UnrealEd/Public/Settings/EditorSettings.h
--- a/Engine/Source/Editor/UnrealEd/Public/Settings/EditorSettings.h
+++ b/Engine/Source/Editor/UnrealEd/Public/Settings/EditorSettings.h
@@ -119,7 +119,7 @@ class UEditorSettings : public UObject
 	// =====================================================================
 
 	/** URL of the Horde server */
-	UPROPERTY(EditAnywhere, config, Category = "Horde", meta = (DisplayName = "Horde Server URL", ConfigRestartRequired = true))
+	UPROPERTY(EditAnywhere, Category = "Horde", meta = (DisplayName = "Horde Server URL", ConfigRestartRequired = true))
 	FString HordeUrl;
 
 	// =====================================================================
