Horde: Fix root path for launching OidcToken. Only affects log messages.

#rnx

diff --git a/Engine/Source/Developer/DesktopPlatform/Private/DesktopPlatformBase.cpp b/Engine/Source/Developer/DesktopPlatform/Private/DesktopPlatformBase.cpp
--- a/Engine/Source/Developer/DesktopPlatform/Private/DesktopPlatformBase.cpp
+++ b/Engine/Source/Developer/DesktopPlatform/Private/DesktopPlatformBase.cpp
@@ -976,7 +976,7 @@ bool FDesktopPlatformBase::GetHordeAccessToken(const FString& HordeUrl, bool bUn
 	}
 
 	FString BaseArguments = FString::Printf(TEXT(" --HordeUrl=\"%s\""), *HordeUrl);
-	return GetOidcAccessTokenInternal(FPaths::EngineDir(), BaseArguments, bUnattended, Warn, OutToken, OutTokenExpiresAt, bOutWasInteractiveLogin);
+	return GetOidcAccessTokenInternal(FPaths::RootDir(), BaseArguments, bUnattended, Warn, OutToken, OutTokenExpiresAt, bOutWasInteractiveLogin);
 }
 
 bool FDesktopPlatformBase::GetOidcAccessTokenInternal(const FString& RootDir, const FString& BaseArguments, bool bUnattended, FFeedbackContext* Warn, FString& OutToken, FDateTime& OutTokenExpiresAt, bool& bOutWasInteractiveLogin)
