[Horde] Add second fallback to find Horde server URL in BaseEngine.ini.

This adds "[Horde]:ServerUrl" section to BaseEngine.ini configuration as second fallback when UE_HORDE_URL is not defined as environment variable nor present in the registry (Windows) or .horde.json file (POSIX).

#rnx

diff --git a/Engine/Source/Developer/DesktopPlatform/Private/DesktopPlatformBase.cpp b/Engine/Source/Developer/DesktopPlatform/Private/DesktopPlatformBase.cpp
--- a/Engine/Source/Developer/DesktopPlatform/Private/DesktopPlatformBase.cpp
+++ b/Engine/Source/Developer/DesktopPlatform/Private/DesktopPlatformBase.cpp
@@ -870,6 +870,7 @@ bool FDesktopPlatformBase::GetOidcTokenStatus(const FString& RootDir, const FStr
 
 FString FDesktopPlatformBase::ReadHordeUrlWithoutCache()
 {
+	// First try to read Horde URL from environment variable
 	FString Url = FPlatformMisc::GetEnvironmentVariable(TEXT("UE_HORDE_URL"));
 	if (!Url.IsEmpty())
 	{
@@ -877,6 +878,7 @@ FString FDesktopPlatformBase::ReadHordeUrlWithoutCache()
 	}
 
 #if PLATFORM_WINDOWS
+	// On Windows, try to read URL from registry entry next
 	if (FWindowsPlatformMisc::QueryRegKey(HKEY_CURRENT_USER, TEXT("SOFTWARE\\Epic Games\\Horde"), TEXT("Url"), Url) && !Url.IsEmpty())
 	{
 		return Url;
@@ -887,6 +889,7 @@ FString FDesktopPlatformBase::ReadHordeUrlWithoutCache()
 		return Url;
 	}
 #else
+	// On POSIX, try to read URL from .horde.json user file next
 	FString FileName = FPaths::Combine(FPlatformProcess::UserHomeDir(), TEXT(".horde.json"));
 
 	FString FileContents;
@@ -902,6 +905,12 @@ FString FDesktopPlatformBase::ReadHordeUrlWithoutCache()
 	}
 #endif
 
+	// As last fallback, try to read Horde URL from INI configuration
+	if (GConfig->GetString(TEXT("Horde"), TEXT("ServerUrl"), Url, GEngineIni))
+	{
+		return Url;
+	}
+
 	return FString();
 }
 
