[Horde] Provide optional parameter to report source of Horde URL configuration.

There are several sources to configure the Horde server URL and this change provides that information in the log.
This adds an optional parameter to IDesktopPlatform::GetHordeUrl for compatibility.

An example of a successful Horde URL query in the log could be as follows:
"Getting Horde server URL succeeded [URL: https://your.horde.server.org/, Source: Engine INI configuration]"

#rnx

diff --git a/Engine/Source/Developer/DesktopPlatform/Private/DesktopPlatformBase.cpp b/Engine/Source/Developer/DesktopPlatform/Private/DesktopPlatformBase.cpp
--- a/Engine/Source/Developer/DesktopPlatform/Private/DesktopPlatformBase.cpp
+++ b/Engine/Source/Developer/DesktopPlatform/Private/DesktopPlatformBase.cpp
@@ -868,12 +868,13 @@ bool FDesktopPlatformBase::GetOidcTokenStatus(const FString& RootDir, const FStr
 	return false;
 }
 
-FString FDesktopPlatformBase::ReadHordeUrlWithoutCache()
+FString FDesktopPlatformBase::ReadHordeUrlWithoutCache(FString& OutHordeUrlConfigSource)
 {
 	// First try to read Horde URL from environment variable
 	FString Url = FPlatformMisc::GetEnvironmentVariable(TEXT("UE_HORDE_URL"));
 	if (!Url.IsEmpty())
 	{
+		OutHordeUrlConfigSource = TEXT("Environment variable \"UE_HORDE_URL\"");
 		return Url;
 	}
 
@@ -881,11 +882,13 @@ FString FDesktopPlatformBase::ReadHordeUrlWithoutCache()
 	// On Windows, try to read URL from registry entry next
 	if (FWindowsPlatformMisc::QueryRegKey(HKEY_CURRENT_USER, TEXT("SOFTWARE\\Epic Games\\Horde"), TEXT("Url"), Url) && !Url.IsEmpty())
 	{
+		OutHordeUrlConfigSource = TEXT("Windows registry \"HKEY_CURRENT_USER\\SOFTWARE\\Epic Games\\Horde\"");
 		return Url;
 	}
 
 	if (FWindowsPlatformMisc::QueryRegKey(HKEY_LOCAL_MACHINE, TEXT("SOFTWARE\\Epic Games\\Horde"), TEXT("Url"), Url) && !Url.IsEmpty())
 	{
+		OutHordeUrlConfigSource = TEXT("Windows registry \"HKEY_LOCAL_MACHINE\\SOFTWARE\\Epic Games\\Horde\"");
 		return Url;
 	}
 #else
@@ -900,6 +903,7 @@ FString FDesktopPlatformBase::ReadHordeUrlWithoutCache()
 		TSharedPtr<FJsonObject> Object;
 		if (FJsonSerializer::Deserialize(Reader, Object) && Object.IsValid() && Object->TryGetStringField(TEXT("Url"), Url))
 		{
+			OutHordeUrlConfigSource = FileName;
 			return Url;
 		}
 	}
@@ -908,21 +912,27 @@ FString FDesktopPlatformBase::ReadHordeUrlWithoutCache()
 	// As last fallback, try to read Horde URL from INI configuration
 	if (GConfig->GetString(TEXT("Horde"), TEXT("ServerUrl"), Url, GEngineIni))
 	{
+		OutHordeUrlConfigSource = FString::Printf(TEXT("%s INI configuration"), *GEngineIni);
 		return Url;
 	}
 
+	OutHordeUrlConfigSource = TEXT("Not found");
 	return FString();
 }
 
-bool FDesktopPlatformBase::GetHordeUrl(FString& OutHordeUrl)
+bool FDesktopPlatformBase::GetHordeUrl(FString& OutHordeUrl, FString* OutHordeUrlConfigSource)
 {
 	if (!bInitializedHordeServerUrl)
 	{
-		HordeServerUrl = ReadHordeUrlWithoutCache();
+		HordeServerUrl = ReadHordeUrlWithoutCache(HordeServerUrlConfigSource);
 		bInitializedHordeServerUrl = true;
 	}
 
 	OutHordeUrl = HordeServerUrl;
+	if (OutHordeUrlConfigSource)
+	{
+		*OutHordeUrlConfigSource = HordeServerUrlConfigSource;
+	}
 	return !OutHordeUrl.IsEmpty();
 }
 
@@ -930,6 +940,7 @@ void FDesktopPlatformBase::SetHordeUrl(const FString& HordeUrl)
 {
 	HordeServerUrl = HordeUrl;
 	bInitializedHordeServerUrl = true;
+	HordeServerUrlConfigSource = TEXT("Unspecified");
 
 #if PLATFORM_WINDOWS
 	if (HordeUrl.IsEmpty())
diff --git a/Engine/Source/Developer/DesktopPlatform/Private/DesktopPlatformBase.h b/Engine/Source/Developer/DesktopPlatform/Private/DesktopPlatformBase.h
--- a/Engine/Source/Developer/DesktopPlatform/Private/DesktopPlatformBase.h
+++ b/Engine/Source/Developer/DesktopPlatform/Private/DesktopPlatformBase.h
@@ -58,7 +58,7 @@ public:
     virtual bool GetOidcAccessToken(const FString& RootDir, const FString& ProjectFileName, const FString& ProviderIdentifier, bool bUnattended, FFeedbackContext* Warn, FString& OutToken, FDateTime& OutTokenExpiresAt, bool& bOutWasInteractiveLogin) override;
 	virtual bool GetOidcTokenStatus(const FString& RootDir, const FString& ProjectFileName, const FString& ProviderIdentifier, FFeedbackContext* Warn, int& OutStatus) override;
 
-	virtual bool GetHordeUrl(FString& OutHordeUrl) override;
+	virtual bool GetHordeUrl(FString& OutHordeUrl, FString* OutHordeUrlConfigSource = nullptr) override;
 	virtual void SetHordeUrl(const FString& HordeUrl) override;
 	virtual bool GetHordeAccessToken(const FString& HordeUrl, bool bUnattended, FFeedbackContext* Warn, FString& OutToken, FDateTime& OutTokenExpiresAt, bool& bOutWasInteractiveLogin) override;
 
@@ -74,10 +74,11 @@ private:
 	TMap<FString, FString> LauncherInstallationList;
 	TMap<FString, FUProjectDictionary> CachedProjectDictionaries;
 	mutable TMap<FString, TArray<FTargetInfo>> ProjectFileToTargets;
-	bool bInitializedHordeServerUrl;
+	bool bInitializedHordeServerUrl = false;
 	FString HordeServerUrl;
+	FString HordeServerUrlConfigSource;
 
-	static FString ReadHordeUrlWithoutCache();
+	static FString ReadHordeUrlWithoutCache(FString& OutHordeUrlConfigSource);
 
 	void ReadLauncherInstallationList();
 	void CheckForLauncherEngineInstallation(const FString &AppId, const FString &Identifier, TMap<FString, FString> &OutInstallations);
diff --git a/Engine/Source/Developer/DesktopPlatform/Public/IDesktopPlatform.h b/Engine/Source/Developer/DesktopPlatform/Public/IDesktopPlatform.h
--- a/Engine/Source/Developer/DesktopPlatform/Public/IDesktopPlatform.h
+++ b/Engine/Source/Developer/DesktopPlatform/Public/IDesktopPlatform.h
@@ -488,13 +488,14 @@ public:
 	/**
 	* Get the URL of the configured Horde server for the current user account.
 	*
-	* @param OutHordeUrl        URL of the Horde server
+	* @param OutHordeUrl				URL of the Horde server
+	* @param OutHordeUrlConfigSource	Optional output parameter describing the source of the configuration.
 	* @return true if a valid URL is returned.
 	*/
-	virtual bool GetHordeUrl(FString& OutHordeUrl) = 0;
+	virtual bool GetHordeUrl(FString& OutHordeUrl, FString* OutHordeUrlConfigSource = nullptr) = 0;
 
 	/**
-	* Set the URL of the configured Horde server for the current user account.
+	* Set the URL of the configured Horde server for the current user account. This will set the source of the configuration to EConfigSource::Unspecified.
 	*
 	* @param HordeUrl        URL of the Horde server
 	*/
diff --git a/Engine/Source/Developer/Horde/Private/Horde.cpp b/Engine/Source/Developer/Horde/Private/Horde.cpp
--- a/Engine/Source/Developer/Horde/Private/Horde.cpp
+++ b/Engine/Source/Developer/Horde/Private/Horde.cpp
@@ -19,14 +19,14 @@ IMPLEMENT_MODULE(FHordeModule, Horde)
 
 // --------------------------------------------------------------------------------
 
-bool FHorde::GetServerUrl(FString& OutUrl)
+bool FHorde::GetServerUrl(FString& OutUrl, FString* OutUrlConfigSource)
 {
 	IDesktopPlatform* DesktopPlatform = FDesktopPlatformModule::TryGet();
 	if (DesktopPlatform == nullptr)
 	{
 		return false;
 	}
-	return DesktopPlatform->GetHordeUrl(OutUrl);
+	return DesktopPlatform->GetHordeUrl(OutUrl, OutUrlConfigSource);
 }
 
 FString FHorde::GetTemplateName()
diff --git a/Engine/Source/Developer/Horde/Public/Horde.h b/Engine/Source/Developer/Horde/Public/Horde.h
--- a/Engine/Source/Developer/Horde/Public/Horde.h
+++ b/Engine/Source/Developer/Horde/Public/Horde.h
@@ -11,7 +11,7 @@
 struct FHorde
 {
 	// Gets the server URL. Equivalent to calling IDesktopPlatform::GetHordeUrl().
-	HORDE_API static bool GetServerUrl(FString& OutUrl);
+	HORDE_API static bool GetServerUrl(FString& OutUrl, FString* OutUrlConfigSource = nullptr);
 
 	// Settings available from the environment when running under a job
 	HORDE_API static FString GetTemplateName();
diff --git a/Engine/Source/Developer/UbaCoordinatorHorde/Private/UbaHordeMetaClient.cpp b/Engine/Source/Developer/UbaCoordinatorHorde/Private/UbaHordeMetaClient.cpp
--- a/Engine/Source/Developer/UbaCoordinatorHorde/Private/UbaHordeMetaClient.cpp
+++ b/Engine/Source/Developer/UbaCoordinatorHorde/Private/UbaHordeMetaClient.cpp
@@ -11,9 +11,14 @@ DEFINE_LOG_CATEGORY(LogUbaHorde);
 
 bool FUbaHordeMetaClient::RefreshHttpClient()
 {
-	if (!FHorde::GetServerUrl(ServerUrl))
+	FString ServerUrlConfigSource;
+	if (FHorde::GetServerUrl(ServerUrl, &ServerUrlConfigSource))
 	{
-		UE_LOG(LogUbaHorde, Warning, TEXT("Getting Horde server URL failed"));
+		UE_LOG(LogUbaHorde, Display, TEXT("Getting Horde server URL succeeded [URL: %s, Source: %s]"), *ServerUrl, *ServerUrlConfigSource);
+	}
+	else
+	{
+		UE_LOG(LogUbaHorde, Warning, TEXT("Getting Horde server URL failed [Source: %s]"), *ServerUrlConfigSource);
 		return false;
 	}
 
@@ -22,7 +27,7 @@ bool FUbaHordeMetaClient::RefreshHttpClient()
 
 	if (!HttpClient->Login(FApp::IsUnattended()))
 	{
-		UE_LOG(LogUbaHorde, Warning, TEXT("Login to Horde server [%s] failed"), *ServerUrl);
+		UE_LOG(LogUbaHorde, Warning, TEXT("Login to Horde server [URL: %s, Source: %s] failed"), *ServerUrl, *ServerUrlConfigSource);
 		return false;
 	}
 
