[DesktopPlatform]
* Explicitly send in basedir when running oidc tool instead of relying on cwd. Reason is because if we run horde through a UE application that has DISABLE_CWD_CHANGES set oidc won't work.

diff --git a/Engine/Source/Developer/DesktopPlatform/Private/DesktopPlatformBase.cpp b/Engine/Source/Developer/DesktopPlatform/Private/DesktopPlatformBase.cpp
--- a/Engine/Source/Developer/DesktopPlatform/Private/DesktopPlatformBase.cpp
+++ b/Engine/Source/Developer/DesktopPlatform/Private/DesktopPlatformBase.cpp
@@ -1255,7 +1255,7 @@ FProcHandle FDesktopPlatformBase::InvokeOidcTokenToolAsync(const FString& InArgu
 	const bool bLaunchHidden = true;
 	const bool bLaunchReallyHidden = bLaunchHidden;
 
-	FProcHandle ProcHandle = FPlatformProcess::CreateProc(*ExecutableFileName, *CmdLineParams, bLaunchDetached, bLaunchHidden, bLaunchReallyHidden, NULL, 0, NULL, OutWritePipe, OutReadPipe);
+	FProcHandle ProcHandle = FPlatformProcess::CreateProc(*ExecutableFileName, *CmdLineParams, bLaunchDetached, bLaunchHidden, bLaunchReallyHidden, NULL, 0, FPlatformProcess::BaseDir(), OutWritePipe, OutReadPipe);
 	if (!ProcHandle.IsValid())
 	{
 		UE_LOG(LogDesktopPlatform, Warning, TEXT("Failed to launch OidcToken (exe path is '%s')"), *ExecutableFileName);
