[DesktopPlatform/UBA]
* Fixed so it is possible to run oidctoken.exe without having the engine environment set up.

diff --git a/Engine/Source/Developer/DesktopPlatform/Private/DesktopPlatformBase.cpp b/Engine/Source/Developer/DesktopPlatform/Private/DesktopPlatformBase.cpp
--- a/Engine/Source/Developer/DesktopPlatform/Private/DesktopPlatformBase.cpp
+++ b/Engine/Source/Developer/DesktopPlatform/Private/DesktopPlatformBase.cpp
@@ -1001,8 +1001,13 @@ bool FDesktopPlatformBase::GetHordeAccessToken(const FString& HordeUrl, bool bUn
 
 bool FDesktopPlatformBase::GetOidcAccessTokenInternal(const FString& RootDir, const FString& BaseArguments, bool bUnattended, FFeedbackContext* Warn, FString& OutToken, FDateTime& OutTokenExpiresAt, bool& bOutWasInteractiveLogin)
 {
+	#if !defined(UE_NO_ENGINE_OIDC) || UE_NO_ENGINE_OIDC == 0
 	IFileManager::Get().MakeDirectory(*FPaths::ProjectIntermediateDir(), /*bTree*/ true);
 	FString ResultFilePath = FPaths::CreateTempFilename(*FPaths::ProjectIntermediateDir(), TEXT("oidcToken.json"));
+	#else
+	FString ResultFilePath = TEXT("oidcToken.json");
+	#endif
+
 	ON_SCOPE_EXIT
 	{
 		IFileManager::Get().Delete(*ResultFilePath, false /* RequireExists */, true /* EvenIfReadOnly */, true /* Quiet */);
diff --git a/Engine/Source/Developer/DesktopPlatform/Private/Windows/DesktopPlatformWindows.cpp b/Engine/Source/Developer/DesktopPlatform/Private/Windows/DesktopPlatformWindows.cpp
--- a/Engine/Source/Developer/DesktopPlatform/Private/Windows/DesktopPlatformWindows.cpp
+++ b/Engine/Source/Developer/DesktopPlatform/Private/Windows/DesktopPlatformWindows.cpp
@@ -609,7 +609,11 @@ int32 FDesktopPlatformWindows::GetShellIntegrationVersion(const FString &FileNam
 
 FString FDesktopPlatformWindows::GetOidcTokenExecutableFilename(const FString& RootDir) const
 {
+	#if !defined(UE_NO_ENGINE_OIDC) || UE_NO_ENGINE_OIDC == 0
 	return FPaths::ConvertRelativePathToFull(RootDir / TEXT("Engine/Binaries/DotNET/OidcToken/win-x64/OidcToken.exe"));
+	#else
+	return TEXT("OidcToken.exe");
+	#endif
 }
 
 #undef LOCTEXT_NAMESPACE
diff --git a/Engine/Source/Programs/UnrealBuildAccelerator/UbaCoordinatorHorde.Target.cs b/Engine/Source/Programs/UnrealBuildAccelerator/UbaCoordinatorHorde.Target.cs
--- a/Engine/Source/Programs/UnrealBuildAccelerator/UbaCoordinatorHorde.Target.cs
+++ b/Engine/Source/Programs/UnrealBuildAccelerator/UbaCoordinatorHorde.Target.cs
@@ -67,6 +67,7 @@ public class UbaCoordinatorHordeTarget : TargetRules
 		GlobalDefinitions.Add("ALLOW_HITCH_DETECTION=0");
 		GlobalDefinitions.Add("UE_USE_MALLOC_FILL_BYTES=0");
 		GlobalDefinitions.Add("DISABLE_CWD_CHANGES=1");
+		GlobalDefinitions.Add("UE_NO_ENGINE_OIDC=1");
 
 		GlobalDefinitions.Add("UBA_COORDINATOR_HORDE_DLL");
 	}
