- Improved the need to call "UBT -Mode=QueryTargets" in the editor, and when it does run, only ask for the games' targets, not every target

diff --git a/Engine/Source/Developer/DesktopPlatform/Private/DesktopPlatformBase.cpp b/Engine/Source/Developer/DesktopPlatform/Private/DesktopPlatformBase.cpp
--- a/Engine/Source/Developer/DesktopPlatform/Private/DesktopPlatformBase.cpp
+++ b/Engine/Source/Developer/DesktopPlatform/Private/DesktopPlatformBase.cpp
@@ -1469,12 +1469,9 @@ const TArray<FTargetInfo>& FDesktopPlatformBase::GetTargetsForProject(const FStr
 	}
 	Arguments += FString::Printf(TEXT(" -Output=\"%s\""), *IFileManager::Get().ConvertToAbsolutePathForExternalAppForRead(*InfoFileName));
 	Arguments += TEXT(" -IncludeAllTargets");
-
-	const bool bIncludeParentAssembly = GConfig->GetBoolOrDefault(TEXT("/Script/BuildSettings.BuildSettings"), TEXT("bIncludeParentAssembly"), true, GEngineIni);
-	if (!bIncludeParentAssembly)
-	{
-		Arguments += TEXT(" -DontIncludeParentAssembly");
-	}
+	// we only want this game's targets, not every program under the sun (note if running UnrealEditor without a project, this will
+	// still list all the programs). This matches what GenerateProjectFiles outputs when it performs this step
+	Arguments += TEXT(" -DontIncludeParentAssembly");
 
 	// Run UBT to update the list of targets. Try to run it without building first.
 	FString Output;
diff --git a/Engine/Source/Programs/UnrealBuildTool/Modes/QueryTargetsMode.cs b/Engine/Source/Programs/UnrealBuildTool/Modes/QueryTargetsMode.cs
--- a/Engine/Source/Programs/UnrealBuildTool/Modes/QueryTargetsMode.cs
+++ b/Engine/Source/Programs/UnrealBuildTool/Modes/QueryTargetsMode.cs
@@ -185,6 +185,8 @@ namespace UnrealBuildTool
 			// Write the output file
 			DirectoryReference.CreateDirectory(OutputFile.Directory);
 			FileReference.WriteAllTextIfDifferent(OutputFile, StringWriter.ToString());
+			// even if the contents don't change, we need to update the last write time (this is safer to do in parallel, unlike WriteAllText)
+			FileReference.SetLastWriteTimeUtc(OutputFile, DateTime.UtcNow);
 		}
 	}
 }
