Expose an API method for fetching OidcToken executable path using an INTERNAL method.  Also ensure that Auth Token acquisition isn't reliant on editor settings INI values when we're building a standalone program.

diff --git a/Engine/Source/Developer/DesktopPlatform/Private/DesktopPlatformBase.cpp b/Engine/Source/Developer/DesktopPlatform/Private/DesktopPlatformBase.cpp
--- a/Engine/Source/Developer/DesktopPlatform/Private/DesktopPlatformBase.cpp
+++ b/Engine/Source/Developer/DesktopPlatform/Private/DesktopPlatformBase.cpp
@@ -1100,9 +1100,13 @@ bool FDesktopPlatformBase::GetOidcAccessTokenInteractive(const FString& RootDir,
 	EAppReturnType::Type userAcknowledgedResult = EAppReturnType::Yes;
 
 	check(GConfig && GConfig->IsReadyForUse());
-	
+
+#if WITH_EDITOR
 	bool bSkipInitialAcknowledgement = false;
 	GConfig->GetBool(TEXT("/Script/UnrealEd.EditorSettings"), TEXT("InteractiveOidcWithoutAcknowledgement"), bSkipInitialAcknowledgement, GEditorSettingsIni);
+#else
+	bool bSkipInitialAcknowledgement = true;
+#endif
 
 	if (!bSkipInitialAcknowledgement)
 	{
diff --git a/Engine/Source/Developer/DesktopPlatform/Private/DesktopPlatformBase.h b/Engine/Source/Developer/DesktopPlatform/Private/DesktopPlatformBase.h
--- a/Engine/Source/Developer/DesktopPlatform/Private/DesktopPlatformBase.h
+++ b/Engine/Source/Developer/DesktopPlatform/Private/DesktopPlatformBase.h
@@ -107,7 +107,6 @@ protected:
 
 	FString GetUnrealBuildToolProjectFileName(const FString& RootDir) const;
 	FString GetUnrealBuildToolExecutableFilename(const FString& RootDir) const;	
-	virtual FString GetOidcTokenExecutableFilename(const FString& RootDir) const = 0;
 };
 
 PRAGMA_ENABLE_DEPRECATION_WARNINGS
diff --git a/Engine/Source/Developer/DesktopPlatform/Public/IDesktopPlatform.h b/Engine/Source/Developer/DesktopPlatform/Public/IDesktopPlatform.h
--- a/Engine/Source/Developer/DesktopPlatform/Public/IDesktopPlatform.h
+++ b/Engine/Source/Developer/DesktopPlatform/Public/IDesktopPlatform.h
@@ -485,6 +485,14 @@ public:
 	*/
 	virtual bool GetOidcTokenStatus(const FString& RootDir, const FString& ProjectFileName, const FString& ProviderIdentifier, FFeedbackContext* Warn, int& OutStatus) = 0;
 
+	/**
+	* Get the path to the OIDC token executable
+
+	* @param RootDir			Engine root directory for the project to use.
+	* @return The path to the OIDC token executable.
+	*/
+	UE_INTERNAL virtual FString GetOidcTokenExecutableFilename(const FString& RootDir) const = 0;
+
 	/**
 	* Get the URL of the configured Horde server for the current user account.
 	*
