Fix for static analysis issue in my last submit.

diff --git a/Engine/Source/Developer/DesktopPlatform/Private/DesktopPlatformBase.cpp b/Engine/Source/Developer/DesktopPlatform/Private/DesktopPlatformBase.cpp
--- a/Engine/Source/Developer/DesktopPlatform/Private/DesktopPlatformBase.cpp
+++ b/Engine/Source/Developer/DesktopPlatform/Private/DesktopPlatformBase.cpp
@@ -1097,16 +1097,12 @@ bool FDesktopPlatformBase::GetOidcAccessTokenInternal(const FString& RootDir, co
 
 bool FDesktopPlatformBase::GetOidcAccessTokenInteractive(const FString& RootDir, const FString& Arguments, bool bUnattended, FFeedbackContext* Warn, int32& OutReturnCode)
 {
-	EAppReturnType::Type userAcknowledgedResult = EAppReturnType::Yes;
-
 	check(GConfig && GConfig->IsReadyForUse());
 
 #if WITH_EDITOR
+	EAppReturnType::Type userAcknowledgedResult = EAppReturnType::Yes;
 	bool bSkipInitialAcknowledgement = false;
 	GConfig->GetBool(TEXT("/Script/UnrealEd.EditorSettings"), TEXT("InteractiveOidcWithoutAcknowledgement"), bSkipInitialAcknowledgement, GEditorSettingsIni);
-#else
-	bool bSkipInitialAcknowledgement = true;
-#endif
 
 	if (!bSkipInitialAcknowledgement)
 	{
@@ -1132,6 +1128,7 @@ bool FDesktopPlatformBase::GetOidcAccessTokenInteractive(const FString& RootDir,
 		OutReturnCode = -1;
 		return false;
 	}
+#endif
 
 	// user has acknowledged the login, we update the editor progress and then we run oidc token to spawn the browser and prompt the login
 
