[DesktopPlatform]
* Moved singleton to source file

diff --git a/Engine/Source/Developer/DesktopPlatform/Private/InstalledPlatformInfo.cpp b/Engine/Source/Developer/DesktopPlatform/Private/InstalledPlatformInfo.cpp
--- a/Engine/Source/Developer/DesktopPlatform/Private/InstalledPlatformInfo.cpp
+++ b/Engine/Source/Developer/DesktopPlatform/Private/InstalledPlatformInfo.cpp
@@ -109,6 +109,12 @@ void FInstalledPlatformInfo::ParsePlatformConfiguration(FString PlatformConfigur
 	}
 }
 
+FInstalledPlatformInfo& FInstalledPlatformInfo::Get()
+{
+	static FInstalledPlatformInfo InfoSingleton;
+	return InfoSingleton;
+}
+
 bool FInstalledPlatformInfo::IsValidConfiguration(const EBuildConfiguration Configuration, EProjectType ProjectType) const
 {
 	return ContainsValidConfiguration(
diff --git a/Engine/Source/Developer/DesktopPlatform/Public/InstalledPlatformInfo.h b/Engine/Source/Developer/DesktopPlatform/Public/InstalledPlatformInfo.h
--- a/Engine/Source/Developer/DesktopPlatform/Public/InstalledPlatformInfo.h
+++ b/Engine/Source/Developer/DesktopPlatform/Public/InstalledPlatformInfo.h
@@ -69,11 +69,7 @@ public:
 	/**
 	 * Accessor for singleton
 	 */
-	static FInstalledPlatformInfo& Get()
-	{
-		static FInstalledPlatformInfo InfoSingleton;
-		return InfoSingleton;
-	}
+	static UE_API FInstalledPlatformInfo& Get();
 
 	/**
 	 * Queries whether a configuration is valid for any available platform
