Added option to get oidc access token with remote configuration

diff --git a/Engine/Source/Developer/DesktopPlatform/Private/DesktopPlatformBase.cpp b/Engine/Source/Developer/DesktopPlatform/Private/DesktopPlatformBase.cpp
--- a/Engine/Source/Developer/DesktopPlatform/Private/DesktopPlatformBase.cpp
+++ b/Engine/Source/Developer/DesktopPlatform/Private/DesktopPlatformBase.cpp
@@ -820,6 +820,14 @@ bool FDesktopPlatformBase::GetOidcAccessToken(const FString& RootDir, const FStr
 	return GetOidcAccessTokenInternal(RootDir, BaseArguments, bUnattended, Warn, OutToken, OutTokenExpiresAt, bOutWasInteractiveLogin);
 }
 
+bool FDesktopPlatformBase::GetOidcAccessTokenWithRemoteConfig(const FString& RootDir, const FString& RemoteHost, bool bUnattended, FFeedbackContext* Warn, FString& OutToken, FDateTime& OutTokenExpiresAt, bool& bOutWasInteractiveLogin)
+{
+	FString BaseArguments;
+	BaseArguments += FString::Printf(TEXT(" --AuthConfigUrl=\"%s\""), *RemoteHost);
+
+	return GetOidcAccessTokenInternal(RootDir, BaseArguments, bUnattended, Warn, OutToken, OutTokenExpiresAt, bOutWasInteractiveLogin);
+}
+
 bool FDesktopPlatformBase::GetOidcTokenStatus(const FString& RootDir, const FString& ProjectFileName, const FString& ProviderIdentifier, FFeedbackContext* Warn, int& OutStatus)
 {
 	IFileManager::Get().MakeDirectory(*FPaths::ProjectIntermediateDir(), /*bTree*/ true);
diff --git a/Engine/Source/Developer/DesktopPlatform/Private/DesktopPlatformBase.h b/Engine/Source/Developer/DesktopPlatform/Private/DesktopPlatformBase.h
--- a/Engine/Source/Developer/DesktopPlatform/Private/DesktopPlatformBase.h
+++ b/Engine/Source/Developer/DesktopPlatform/Private/DesktopPlatformBase.h
@@ -56,6 +56,7 @@ public:
 	virtual bool IsUnrealBuildToolRunning() override;
 
     virtual bool GetOidcAccessToken(const FString& RootDir, const FString& ProjectFileName, const FString& ProviderIdentifier, bool bUnattended, FFeedbackContext* Warn, FString& OutToken, FDateTime& OutTokenExpiresAt, bool& bOutWasInteractiveLogin) override;
+	virtual bool GetOidcAccessTokenWithRemoteConfig(const FString& RootDir, const FString& RemoteHost, bool bUnattended, FFeedbackContext* Warn, FString& OutToken, FDateTime& OutTokenExpiresAt, bool& bOutWasInteractiveLogin) override;
 	virtual bool GetOidcTokenStatus(const FString& RootDir, const FString& ProjectFileName, const FString& ProviderIdentifier, FFeedbackContext* Warn, int& OutStatus) override;
 
 	virtual bool GetHordeUrl(FString& OutHordeUrl, FString* OutHordeUrlConfigSource = nullptr) override;
diff --git a/Engine/Source/Developer/DesktopPlatform/Public/IDesktopPlatform.h b/Engine/Source/Developer/DesktopPlatform/Public/IDesktopPlatform.h
--- a/Engine/Source/Developer/DesktopPlatform/Public/IDesktopPlatform.h
+++ b/Engine/Source/Developer/DesktopPlatform/Public/IDesktopPlatform.h
@@ -473,6 +473,20 @@ public:
 	*/
 	virtual bool GetOidcAccessToken(const FString& RootDir, const FString& ProjectFileName, const FString& ProviderIdentifier, bool bUnattended, FFeedbackContext* Warn, FString& OutToken, FDateTime& OutTokenExpiresAt, bool& bOutWasInteractiveLogin) = 0;
 
+	/**
+	* Get a Access token for use by OIDC compliant Identity Providers using a remote configuration
+
+	* @param RootDir			Engine root directory for the project to use.
+	* @param RemoteHost			The url to the remote host to fetch the config from
+	* @param bUnattended		True to indicate that no user interaction should be assumed
+	* @param Warn				Feedback context to use for progress updates
+	* @param OutToken			The allocated access token
+	* @param OutTokenExpiresAt	When the token expires
+	* @param bOutWasInteractiveLogin	True if the interactive login flow was used
+	* @return true if the task completed successfully.
+	*/
+	virtual bool GetOidcAccessTokenWithRemoteConfig(const FString& RootDir, const FString& RemoteHost, bool bUnattended, FFeedbackContext* Warn, FString& OutToken, FDateTime& OutTokenExpiresAt, bool& bOutWasInteractiveLogin) = 0;
+
 	/**
 	* Get the status of a access token for use by OIDC compliant Identity Providers
 
diff --git a/Engine/Source/Developer/Zen/Private/BuildServerInterface.cpp b/Engine/Source/Developer/Zen/Private/BuildServerInterface.cpp
--- a/Engine/Source/Developer/Zen/Private/BuildServerInterface.cpp
+++ b/Engine/Source/Developer/Zen/Private/BuildServerInterface.cpp
@@ -631,7 +631,7 @@ FBuildServiceInstance::AcquireAccessToken()
 	bool bWasInteractiveLogin = false;
 
 	IDesktopPlatform* DesktopPlatform = FDesktopPlatformModule::TryGet();
-	if (DesktopPlatform && DesktopPlatform->GetOidcAccessToken(FPaths::RootDir(), {}, Settings.GetOAuthProviderIdentifier(), FApp::IsUnattended(), GWarn, AccessTokenString, TokenExpiresAt, bWasInteractiveLogin))
+	if (DesktopPlatform && DesktopPlatform->GetOidcAccessTokenWithRemoteConfig(FPaths::RootDir(), *WriteToString<64>(EffectiveDomain), FApp::IsUnattended(), GWarn, AccessTokenString, TokenExpiresAt, bWasInteractiveLogin))
 	{
 		if (bWasInteractiveLogin)
 		{
