Add means to explicitly specify interaction mode when obtaining OIDC tokens.  This is then used to allow the BuildStorageTool a means to implement "login non-interactively if you can" functionality on startup.

diff --git a/Engine/Source/Developer/DerivedDataCache/Private/HttpCacheStore.cpp b/Engine/Source/Developer/DerivedDataCache/Private/HttpCacheStore.cpp
--- a/Engine/Source/Developer/DerivedDataCache/Private/HttpCacheStore.cpp
+++ b/Engine/Source/Developer/DerivedDataCache/Private/HttpCacheStore.cpp
@@ -2558,7 +2558,8 @@ bool FHttpCacheStore::AcquireAccessToken(IHttpClient* Client)
 		bool bWasInteractiveLogin = false;
 
 		IDesktopPlatform* DesktopPlatform = FDesktopPlatformModule::TryGet();
-		if (DesktopPlatform && DesktopPlatform->GetOidcAccessToken(FPaths::RootDir(), FPaths::GetProjectFilePath(), OAuthProviderIdentifier, FApp::IsUnattended(), GWarn, AccessTokenString, TokenExpiresAt, bWasInteractiveLogin))
+		EDesktopLoginInteractionLevel InteractionLevel = FApp::IsUnattended() ? EDesktopLoginInteractionLevel::TimeLimited : EDesktopLoginInteractionLevel::Interactive;
+		if (DesktopPlatform && DesktopPlatform->GetOidcAccessToken(FPaths::RootDir(), FPaths::GetProjectFilePath(), OAuthProviderIdentifier, InteractionLevel, GWarn, AccessTokenString, TokenExpiresAt, bWasInteractiveLogin))
 		{
 			if (bWasInteractiveLogin)
 			{
diff --git a/Engine/Source/Developer/DesktopPlatform/Private/DesktopPlatformBase.cpp b/Engine/Source/Developer/DesktopPlatform/Private/DesktopPlatformBase.cpp
--- a/Engine/Source/Developer/DesktopPlatform/Private/DesktopPlatformBase.cpp
+++ b/Engine/Source/Developer/DesktopPlatform/Private/DesktopPlatformBase.cpp
@@ -808,7 +808,7 @@ bool FDesktopPlatformBase::IsUnrealBuildToolRunning()
 	return bIsRunning;
 }
 
-bool FDesktopPlatformBase::GetOidcAccessToken(const FString& RootDir, const FString& ProjectFileName, const FString& ProviderIdentifier, bool bUnattended, FFeedbackContext* Warn, FString& OutToken, FDateTime& OutTokenExpiresAt, bool& bOutWasInteractiveLogin)
+bool FDesktopPlatformBase::GetOidcAccessToken(const FString& RootDir, const FString& ProjectFileName, const FString& ProviderIdentifier, EDesktopLoginInteractionLevel InteractionLevel, FFeedbackContext* Warn, FString& OutToken, FDateTime& OutTokenExpiresAt, bool& bOutWasInteractiveLogin)
 {
 	FString BaseArguments;
 	BaseArguments += FString::Printf(TEXT(" --Service=\"%s\""), *ProviderIdentifier);
@@ -817,15 +817,27 @@ bool FDesktopPlatformBase::GetOidcAccessToken(const FString& RootDir, const FStr
 		BaseArguments += FString::Printf(TEXT(" --project=\"%s\""), *IFileManager::Get().ConvertToAbsolutePathForExternalAppForRead(*FPaths::GetPath(*ProjectFileName)));
 	}
 
-	return GetOidcAccessTokenInternal(RootDir, BaseArguments, bUnattended, Warn, OutToken, OutTokenExpiresAt, bOutWasInteractiveLogin);
+	return GetOidcAccessTokenInternal(RootDir, BaseArguments, InteractionLevel, Warn, OutToken, OutTokenExpiresAt, bOutWasInteractiveLogin);
 }
 
-bool FDesktopPlatformBase::GetOidcAccessTokenWithRemoteConfig(const FString& RootDir, const FString& RemoteHost, bool bUnattended, FFeedbackContext* Warn, FString& OutToken, FDateTime& OutTokenExpiresAt, bool& bOutWasInteractiveLogin)
+bool FDesktopPlatformBase::GetOidcAccessTokenWithRemoteConfig(const FString& RootDir, const FString& RemoteHost, EDesktopLoginInteractionLevel InteractionLevel, FFeedbackContext* Warn, FString& OutToken, FDateTime& OutTokenExpiresAt, bool& bOutWasInteractiveLogin)
 {
 	FString BaseArguments;
 	BaseArguments += FString::Printf(TEXT(" --AuthConfigUrl=\"%s\""), *RemoteHost);
 
-	return GetOidcAccessTokenInternal(RootDir, BaseArguments, bUnattended, Warn, OutToken, OutTokenExpiresAt, bOutWasInteractiveLogin);
+	return GetOidcAccessTokenInternal(RootDir, BaseArguments, InteractionLevel, Warn, OutToken, OutTokenExpiresAt, bOutWasInteractiveLogin);
+}
+
+bool FDesktopPlatformBase::GetOidcAccessToken(const FString& RootDir, const FString& ProjectFileName, const FString& ProviderIdentifier, bool bUnattended, FFeedbackContext* Warn, FString& OutToken, FDateTime& OutTokenExpiresAt, bool& bOutWasInteractiveLogin)
+{
+	const EDesktopLoginInteractionLevel InteractionLevel = bUnattended ? EDesktopLoginInteractionLevel::None : EDesktopLoginInteractionLevel::Interactive;
+	return GetOidcAccessToken(RootDir, ProjectFileName, ProviderIdentifier, InteractionLevel, Warn, OutToken, OutTokenExpiresAt, bOutWasInteractiveLogin);
+}
+
+bool FDesktopPlatformBase::GetOidcAccessTokenWithRemoteConfig(const FString& RootDir, const FString& RemoteHost, bool bUnattended, FFeedbackContext* Warn, FString& OutToken, FDateTime& OutTokenExpiresAt, bool& bOutWasInteractiveLogin)
+{
+	const EDesktopLoginInteractionLevel InteractionLevel = bUnattended ? EDesktopLoginInteractionLevel::None : EDesktopLoginInteractionLevel::Interactive;
+	return GetOidcAccessTokenWithRemoteConfig(RootDir, RemoteHost, InteractionLevel, Warn, OutToken, OutTokenExpiresAt, bOutWasInteractiveLogin);
 }
 
 bool FDesktopPlatformBase::GetOidcTokenStatus(const FString& RootDir, const FString& ProjectFileName, const FString& ProviderIdentifier, FFeedbackContext* Warn, int& OutStatus)
@@ -989,7 +1001,7 @@ void FDesktopPlatformBase::SetHordeUrl(const FString& HordeUrl)
 #endif
 }
 
-bool FDesktopPlatformBase::GetHordeAccessToken(const FString& HordeUrl, bool bUnattended, FFeedbackContext* Warn, FString& OutToken, FDateTime& OutTokenExpiresAt, bool& bOutWasInteractiveLogin)
+bool FDesktopPlatformBase::GetHordeAccessToken(const FString& HordeUrl, EDesktopLoginInteractionLevel InteractionLevel, FFeedbackContext* Warn, FString& OutToken, FDateTime& OutTokenExpiresAt, bool& bOutWasInteractiveLogin)
 {
 	FString Token = FPlatformMisc::GetEnvironmentVariable(TEXT("UE_HORDE_TOKEN"));
 	if (!Token.IsEmpty())
@@ -1004,10 +1016,16 @@ bool FDesktopPlatformBase::GetHordeAccessToken(const FString& HordeUrl, bool bUn
 	}
 
 	FString BaseArguments = FString::Printf(TEXT(" --HordeUrl=\"%s\""), *HordeUrl);
-	return GetOidcAccessTokenInternal(FPaths::RootDir(), BaseArguments, bUnattended, Warn, OutToken, OutTokenExpiresAt, bOutWasInteractiveLogin);
+	return GetOidcAccessTokenInternal(FPaths::RootDir(), BaseArguments, InteractionLevel, Warn, OutToken, OutTokenExpiresAt, bOutWasInteractiveLogin);
 }
 
-bool FDesktopPlatformBase::GetOidcAccessTokenInternal(const FString& RootDir, const FString& BaseArguments, bool bUnattended, FFeedbackContext* Warn, FString& OutToken, FDateTime& OutTokenExpiresAt, bool& bOutWasInteractiveLogin)
+bool FDesktopPlatformBase::GetHordeAccessToken(const FString& HordeUrl, bool bUnattended, FFeedbackContext* Warn, FString& OutToken, FDateTime& OutTokenExpiresAt, bool& bOutWasInteractiveLogin)
+{
+	const EDesktopLoginInteractionLevel InteractionLevel = bUnattended ? EDesktopLoginInteractionLevel::None : EDesktopLoginInteractionLevel::Interactive;
+	return GetHordeAccessToken(HordeUrl, InteractionLevel, Warn, OutToken, OutTokenExpiresAt, bOutWasInteractiveLogin);
+}
+
+bool FDesktopPlatformBase::GetOidcAccessTokenInternal(const FString& RootDir, const FString& BaseArguments, EDesktopLoginInteractionLevel InteractionLevel, FFeedbackContext* Warn, FString& OutToken, FDateTime& OutTokenExpiresAt, bool& bOutWasInteractiveLogin)
 {
 	#if !defined(UE_NO_ENGINE_OIDC) || UE_NO_ENGINE_OIDC == 0
 	IFileManager::Get().MakeDirectory(*FPaths::ProjectIntermediateDir(), /*bTree*/ true);
@@ -1037,15 +1055,15 @@ bool FDesktopPlatformBase::GetOidcAccessTokenInternal(const FString& RootDir, co
 
 	if (ExitCode == 10)
 	{
-		bRes = GetOidcAccessTokenInteractive(RootDir, Arguments, bUnattended, Warn, ExitCode);
+		bRes = GetOidcAccessTokenInteractive(RootDir, Arguments, InteractionLevel, Warn, ExitCode);
 
 		bOutWasInteractiveLogin = true;
 
 		if (!bRes)
 		{
-			if (bUnattended)
+			if (InteractionLevel != EDesktopLoginInteractionLevel::Interactive)
 			{
-				UE_LOG(LogDesktopPlatform, Warning, TEXT("Unable to allocate an access token. Unattended set so unable to complete interactive login. Make sure you start the editor and login once or log in using UGS or using the UGS cli command 'login'. (Ran '%s%s')"), *GetOidcTokenExecutableFilename(RootDir), *Arguments);
+				UE_LOG(LogDesktopPlatform, Warning, TEXT("Unable to allocate an access token. Unable to complete interactive login. Make sure you start the editor and login once or log in using UGS or using the UGS cli command 'login'. (Ran '%s%s')"), *GetOidcTokenExecutableFilename(RootDir), *Arguments);
 			}
 			else
 			{
@@ -1103,18 +1121,24 @@ bool FDesktopPlatformBase::GetOidcAccessTokenInternal(const FString& RootDir, co
 	return false;
 }
 
-bool FDesktopPlatformBase::GetOidcAccessTokenInteractive(const FString& RootDir, const FString& Arguments, bool bUnattended, FFeedbackContext* Warn, int32& OutReturnCode)
+bool FDesktopPlatformBase::GetOidcAccessTokenInteractive(const FString& RootDir, const FString& Arguments, EDesktopLoginInteractionLevel InteractionLevel, FFeedbackContext* Warn, int32& OutReturnCode)
 {
-	check(GConfig && GConfig->IsReadyForUse());
+	if (InteractionLevel == EDesktopLoginInteractionLevel::None)
+	{
+		OutReturnCode = -1;
+		return false;
+	}
 
 #if WITH_EDITOR
+	check(GConfig && GConfig->IsReadyForUse());
+
 	EAppReturnType::Type userAcknowledgedResult = EAppReturnType::Yes;
 	bool bSkipInitialAcknowledgement = false;
 	GConfig->GetBool(TEXT("/Script/UnrealEd.EditorSettings"), TEXT("InteractiveOidcWithoutAcknowledgement"), bSkipInitialAcknowledgement, GEditorSettingsIni);
 
 	if (!bSkipInitialAcknowledgement)
 	{
-		if (bUnattended)
+		if (InteractionLevel != EDesktopLoginInteractionLevel::Interactive)
 		{
 			OutReturnCode = -1;
 			return false;
@@ -1182,7 +1206,7 @@ bool FDesktopPlatformBase::GetOidcAccessTokenInteractive(const FString& RootDir,
 		// once we have waited for 30 seconds without success we give the user a option to abort, or if unattended, abort without prompting
 		else if (WaitDuration > 30.0 && DurationPhase == EWaitDurationPhase::Prompt)
 		{
-			if (bUnattended)
+			if (InteractionLevel == EDesktopLoginInteractionLevel::TimeLimited)
 			{
 				bIsFinished = !FPlatformProcess::IsProcRunning(ProcHandle);
 				break;
diff --git a/Engine/Source/Developer/DesktopPlatform/Private/DesktopPlatformBase.h b/Engine/Source/Developer/DesktopPlatform/Private/DesktopPlatformBase.h
--- a/Engine/Source/Developer/DesktopPlatform/Private/DesktopPlatformBase.h
+++ b/Engine/Source/Developer/DesktopPlatform/Private/DesktopPlatformBase.h
@@ -55,12 +55,15 @@ public:
 	virtual bool RunUnrealBuildTool(const FText& Description, const FString& RootDir, const FString& Arguments, FFeedbackContext* Warn) override final;
 	virtual bool IsUnrealBuildToolRunning() override;
 
+    virtual bool GetOidcAccessToken(const FString& RootDir, const FString& ProjectFileName, const FString& ProviderIdentifier, EDesktopLoginInteractionLevel InteractionLevel, FFeedbackContext* Warn, FString& OutToken, FDateTime& OutTokenExpiresAt, bool& bOutWasInteractiveLogin) override;
+	virtual bool GetOidcAccessTokenWithRemoteConfig(const FString& RootDir, const FString& RemoteHost, EDesktopLoginInteractionLevel InteractionLevel, FFeedbackContext* Warn, FString& OutToken, FDateTime& OutTokenExpiresAt, bool& bOutWasInteractiveLogin) override;
     virtual bool GetOidcAccessToken(const FString& RootDir, const FString& ProjectFileName, const FString& ProviderIdentifier, bool bUnattended, FFeedbackContext* Warn, FString& OutToken, FDateTime& OutTokenExpiresAt, bool& bOutWasInteractiveLogin) override;
 	virtual bool GetOidcAccessTokenWithRemoteConfig(const FString& RootDir, const FString& RemoteHost, bool bUnattended, FFeedbackContext* Warn, FString& OutToken, FDateTime& OutTokenExpiresAt, bool& bOutWasInteractiveLogin) override;
 	virtual bool GetOidcTokenStatus(const FString& RootDir, const FString& ProjectFileName, const FString& ProviderIdentifier, FFeedbackContext* Warn, int& OutStatus) override;
 
 	virtual bool GetHordeUrl(FString& OutHordeUrl, FString* OutHordeUrlConfigSource = nullptr) override;
 	virtual void SetHordeUrl(const FString& HordeUrl) override;
+	virtual bool GetHordeAccessToken(const FString& HordeUrl, EDesktopLoginInteractionLevel InteractionLevel, FFeedbackContext* Warn, FString& OutToken, FDateTime& OutTokenExpiresAt, bool& bOutWasInteractiveLogin) override;
 	virtual bool GetHordeAccessToken(const FString& HordeUrl, bool bUnattended, FFeedbackContext* Warn, FString& OutToken, FDateTime& OutTokenExpiresAt, bool& bOutWasInteractiveLogin) override;
 
 	virtual const TArray<FTargetInfo>& GetTargetsForProject(const FString& ProjectFile) const override;
@@ -95,8 +98,8 @@ private:
 
 	bool InvokeOidcTokenToolSync(const FText& Description, const FString& RootDir, const FString& Arguments, FFeedbackContext* Warn, int32& OutReturnCode, FString& OutProcOutput);
 	FProcHandle InvokeOidcTokenToolAsync(const FString& InArguments, void*& OutReadPipe, void*& OutWritePipe);
-	bool GetOidcAccessTokenInternal(const FString& RootDir, const FString& Arguments, bool bUnattended, FFeedbackContext* Warn, FString& OutToken, FDateTime& OutTokenExpiresAt, bool& bOutWasInteractiveLogin);
-	bool GetOidcAccessTokenInteractive(const FString& RootDir, const FString& Arguments, bool bUnattended, FFeedbackContext* Warn, int32& OutReturnCode);
+	bool GetOidcAccessTokenInternal(const FString& RootDir, const FString& Arguments, EDesktopLoginInteractionLevel InteractionLevel, FFeedbackContext* Warn, FString& OutToken, FDateTime& OutTokenExpiresAt, bool& bOutWasInteractiveLogin);
+	bool GetOidcAccessTokenInteractive(const FString& RootDir, const FString& Arguments, EDesktopLoginInteractionLevel InteractionLevel, FFeedbackContext* Warn, int32& OutReturnCode);
 
 	static bool ReadTargetInfo(const FString& FileName, TArray<FTargetInfo>& Targets);
 
diff --git a/Engine/Source/Developer/DesktopPlatform/Public/IDesktopPlatform.h b/Engine/Source/Developer/DesktopPlatform/Public/IDesktopPlatform.h
--- a/Engine/Source/Developer/DesktopPlatform/Public/IDesktopPlatform.h
+++ b/Engine/Source/Developer/DesktopPlatform/Public/IDesktopPlatform.h
@@ -36,6 +36,12 @@ enum class EFontImportFlags
 
 ENUM_CLASS_FLAGS(EFontImportFlags)
 
+enum class EDesktopLoginInteractionLevel : uint8
+{
+	None,			// No user interaction permitted
+	TimeLimited,	// Allow user interaction within a limited window of time (eg: 30 seconds), but don't wait indefinitely for interaction
+	Interactive		// Allow user interaction with an indefinite wait if necessary
+};
 
 /**
  * When constructed leaves system wide modal mode (all windows disabled except for the OS modal window)
@@ -474,6 +480,36 @@ public:
 	 * @param bOutWasInteractiveLogin	True if the interactive login flow was used
 	 * @return true if the task completed successfully.
 	 */
+	virtual bool GetOidcAccessToken(const FString& RootDir, const FString& ProjectFileName, const FString& ProviderIdentifier, EDesktopLoginInteractionLevel InteractionLevel, FFeedbackContext* Warn, FString& OutToken, FDateTime& OutTokenExpiresAt, bool& bOutWasInteractiveLogin) = 0;
+
+	/**
+	 * Get an access token for use by OIDC compliant Identity Providers using a remote configuration
+
+	 * @param RootDir					Engine root directory for the project to use.
+	 * @param RemoteHost				The url to the remote host to fetch the config from
+	 * @param bUnattended				True to indicate that no user interaction should be assumed
+	 * @param Warn						Feedback context to use for progress updates
+	 * @param OutToken					The allocated access token
+	 * @param OutTokenExpiresAt			When the token expires
+	 * @param bOutWasInteractiveLogin	True if the interactive login flow was used
+	 * @return true if the task completed successfully.
+	 */
+	virtual bool GetOidcAccessTokenWithRemoteConfig(const FString& RootDir, const FString& RemoteHost, EDesktopLoginInteractionLevel InteractionLevel, FFeedbackContext* Warn, FString& OutToken, FDateTime& OutTokenExpiresAt, bool& bOutWasInteractiveLogin) = 0;
+
+	/**
+	 * Get an access token for use by OIDC compliant Identity Providers
+
+	 * @param RootDir					Engine root directory for the project to use.
+	 * @param ProjectFileName			Filename of the current project
+	 * @param ProviderIdentifier		The identifier for the provider as configured in oidc-token.json
+	 * @param bUnattended				True to indicate that no user interaction should be assumed
+	 * @param Warn						Feedback context to use for progress updates
+	 * @param OutToken					The allocated access token
+	 * @param OutTokenExpiresAt			When the token expires
+	 * @param bOutWasInteractiveLogin	True if the interactive login flow was used
+	 * @return true if the task completed successfully.
+	 */
+	UE_DEPRECATED(5.7, "Please use overload that takes EDesktopLoginInteractionLevel instead")
 	virtual bool GetOidcAccessToken(const FString& RootDir, const FString& ProjectFileName, const FString& ProviderIdentifier, bool bUnattended, FFeedbackContext* Warn, FString& OutToken, FDateTime& OutTokenExpiresAt, bool& bOutWasInteractiveLogin) = 0;
 
 	/**
@@ -488,6 +524,7 @@ public:
 	 * @param bOutWasInteractiveLogin	True if the interactive login flow was used
 	 * @return true if the task completed successfully.
 	 */
+	UE_DEPRECATED(5.7, "Please use overload that takes EDesktopLoginInteractionLevel instead")
 	virtual bool GetOidcAccessTokenWithRemoteConfig(const FString& RootDir, const FString& RemoteHost, bool bUnattended, FFeedbackContext* Warn, FString& OutToken, FDateTime& OutTokenExpiresAt, bool& bOutWasInteractiveLogin) = 0;
 
 	/**
@@ -537,5 +574,19 @@ public:
 	 * @param bOutWasInteractiveLogin	True if the interactive login flow was used
 	 * @return true if the task completed successfully.
 	 */
+	virtual bool GetHordeAccessToken(const FString& HordeUrl, EDesktopLoginInteractionLevel InteractionLevel, FFeedbackContext* Warn, FString& OutToken, FDateTime& OutTokenExpiresAt, bool& bOutWasInteractiveLogin) = 0;
+
+	/**
+	 * Gets an access token for the given Horde Server.
+	 *
+	 * @param HordeUrl					URL of the Horde server.
+	 * @param bUnattended				True to indicate that no user interaction should be assumed
+	 * @param Warn						Feedback context to use for progress updates
+	 * @param OutToken					The allocated access token
+	 * @param OutTokenExpiresAt			When the token expires
+	 * @param bOutWasInteractiveLogin	True if the interactive login flow was used
+	 * @return true if the task completed successfully.
+	 */
+	UE_DEPRECATED(5.7, "Please use overload that takes EDesktopLoginInteractionLevel instead")
 	virtual bool GetHordeAccessToken(const FString& HordeUrl, bool bUnattended, FFeedbackContext* Warn, FString& OutToken, FDateTime& OutTokenExpiresAt, bool& bOutWasInteractiveLogin) = 0;
 };
diff --git a/Engine/Source/Developer/Horde/Private/HordeHttpClient.cpp b/Engine/Source/Developer/Horde/Private/HordeHttpClient.cpp
--- a/Engine/Source/Developer/Horde/Private/HordeHttpClient.cpp
+++ b/Engine/Source/Developer/Horde/Private/HordeHttpClient.cpp
@@ -27,7 +27,8 @@ bool FHordeHttpClient::Login(bool bUnattended, FFeedbackContext* Warn)
 		FDateTime ExpiresAt;
 		bool bWasInteractive = false;
 
-		if (DesktopPlatform->GetHordeAccessToken(ServerUrl, bUnattended, Warn, NewToken, ExpiresAt, bWasInteractive))
+		EDesktopLoginInteractionLevel InteractionLevel = bUnattended ? EDesktopLoginInteractionLevel::TimeLimited : EDesktopLoginInteractionLevel::Interactive;
+		if (DesktopPlatform->GetHordeAccessToken(ServerUrl, InteractionLevel, Warn, NewToken, ExpiresAt, bWasInteractive))
 		{
 			Token = NewToken;
 			return true;
@@ -45,7 +46,8 @@ bool FHordeHttpClient::LoginWithOidc(const TCHAR* Profile, bool bUnattended, FFe
 		FDateTime ExpiresAt;
 		bool bWasInteractive = false;
 
-		if (DesktopPlatform->GetOidcAccessToken(FPaths::RootDir(), FString(), Profile, bUnattended, Warn, NewToken, ExpiresAt, bWasInteractive))
+		EDesktopLoginInteractionLevel InteractionLevel = bUnattended ? EDesktopLoginInteractionLevel::TimeLimited : EDesktopLoginInteractionLevel::Interactive;
+		if (DesktopPlatform->GetOidcAccessToken(FPaths::RootDir(), FString(), Profile, InteractionLevel, Warn, NewToken, ExpiresAt, bWasInteractive))
 		{
 			Token = NewToken;
 			return true;
diff --git a/Engine/Source/Developer/Zen/Private/BuildListRetriever.cpp b/Engine/Source/Developer/Zen/Private/BuildListRetriever.cpp
--- a/Engine/Source/Developer/Zen/Private/BuildListRetriever.cpp
+++ b/Engine/Source/Developer/Zen/Private/BuildListRetriever.cpp
@@ -1,6 +1,7 @@
 // Copyright Epic Games, Inc. All Rights Reserved.
 
 #include "Experimental/BuildListRetriever.h"
+#include "Misc/App.h"
 #include "Misc/DateTime.h"
 #include "Misc/EngineVersion.h"
 
@@ -19,7 +20,9 @@ void FBuildListRetriever::ConnectToBuildService()
 
 	bConnected = true;
 	ServiceInstance->OnRefreshNamespacesAndBucketsComplete().AddSP(this, &FBuildListRetriever::OnNamespacesAndBucketsRefreshed);
-	ServiceInstance->Connect([this](UE::Zen::Build::FBuildServiceInstance::EConnectionState ConnectionState)
+	ServiceInstance->Connect(!FApp::IsUnattended(), [this]
+		(UE::Zen::Build::FBuildServiceInstance::EConnectionState ConnectionState,
+			UE::Zen::Build::FBuildServiceInstance::EConnectionFailureReason FailureReason)
 		{
 			if (ConnectionState == UE::Zen::Build::FBuildServiceInstance::EConnectionState::ConnectionSucceeded)
 			{
diff --git a/Engine/Source/Developer/Zen/Private/BuildServerInterface.cpp b/Engine/Source/Developer/Zen/Private/BuildServerInterface.cpp
--- a/Engine/Source/Developer/Zen/Private/BuildServerInterface.cpp
+++ b/Engine/Source/Developer/Zen/Private/BuildServerInterface.cpp
@@ -222,19 +222,18 @@ FBuildServiceInstance::~FBuildServiceInstance()
 }
 
 void
-FBuildServiceInstance::Connect(FOnConnectionComplete&& OnConnectionComplete)
+FBuildServiceInstance::Connect(bool bAllowInteractive, FOnConnectionComplete&& OnConnectionComplete)
 {
-	EConnectionState ExpectedConnectionState = EConnectionState::NotStarted;
-	if (!ConnectionState.compare_exchange_weak(ExpectedConnectionState, EConnectionState::ConnectionInProgress))
+	if (ConnectionState.exchange(EConnectionState::ConnectionInProgress, std::memory_order_relaxed) == EConnectionState::ConnectionInProgress)
 	{
 		if (OnConnectionComplete)
 		{
-			OnConnectionComplete(ExpectedConnectionState);
+			OnConnectionComplete(EConnectionState::ConnectionFailed, EConnectionFailureReason::UnexpectedState);
 		}
 		return;
 	}
 
-	UE::Tasks::Launch(UE_SOURCE_LOCATION, [this, OnConnectionComplete = MoveTemp(OnConnectionComplete)]()
+	UE::Tasks::Launch(UE_SOURCE_LOCATION, [this, bAllowInteractive, OnConnectionComplete = MoveTemp(OnConnectionComplete)]()
 	{
 		TAnsiStringBuilder<256> ResolvedHost;
 		double ResolvedLatency;
@@ -245,12 +244,13 @@ FBuildServiceInstance::Connect(FOnConnectionComplete&& OnConnectionComplete)
 			EffectiveDomain.Reset();
 			EffectiveDomain.Append(ResolvedHost);
 
-			if (AcquireAccessToken())
+			EDesktopLoginInteractionLevel InteractionLevel = bAllowInteractive ? EDesktopLoginInteractionLevel::Interactive : EDesktopLoginInteractionLevel::None;
+			if (AcquireAccessToken(InteractionLevel))
 			{
 				ConnectionState.store(EConnectionState::ConnectionSucceeded, std::memory_order_relaxed);
 				if (OnConnectionComplete)
 				{
-					OnConnectionComplete(EConnectionState::ConnectionSucceeded);
+					OnConnectionComplete(EConnectionState::ConnectionSucceeded, EConnectionFailureReason::None);
 				}
 			}
 			else
@@ -259,7 +259,7 @@ FBuildServiceInstance::Connect(FOnConnectionComplete&& OnConnectionComplete)
 				ConnectionState.store(EConnectionState::ConnectionFailed, std::memory_order_relaxed);
 				if (OnConnectionComplete)
 				{
-					OnConnectionComplete(EConnectionState::ConnectionFailed);
+					OnConnectionComplete(EConnectionState::ConnectionFailed, EConnectionFailureReason::FailedToAcquireAccessToken);
 				}
 			}
 		}
@@ -274,7 +274,7 @@ FBuildServiceInstance::Connect(FOnConnectionComplete&& OnConnectionComplete)
 			ConnectionState.store(EConnectionState::ConnectionFailed, std::memory_order_relaxed);
 			if (OnConnectionComplete)
 			{
-				OnConnectionComplete(EConnectionState::ConnectionFailed);
+				OnConnectionComplete(EConnectionState::ConnectionFailed, EConnectionFailureReason::FailedToResolveHost);
 			}
 		}
 
@@ -670,7 +670,7 @@ FBuildServiceInstance::Initialize()
 }
 
 bool
-FBuildServiceInstance::AcquireAccessToken()
+FBuildServiceInstance::AcquireAccessToken(EDesktopLoginInteractionLevel InteractionLevel)
 {
 	if (GetEffectiveDomain().StartsWith("http://localhost"))
 	{
@@ -705,7 +705,7 @@ FBuildServiceInstance::AcquireAccessToken()
 	bool bWasInteractiveLogin = false;
 
 	IDesktopPlatform* DesktopPlatform = FDesktopPlatformModule::TryGet();
-	if (DesktopPlatform && DesktopPlatform->GetOidcAccessTokenWithRemoteConfig(FPaths::RootDir(), *WriteToString<64>(EffectiveDomain), FApp::IsUnattended(), GWarn, AccessTokenString, TokenExpiresAt, bWasInteractiveLogin))
+	if (DesktopPlatform && DesktopPlatform->GetOidcAccessTokenWithRemoteConfig(FPaths::RootDir(), *WriteToString<64>(EffectiveDomain), InteractionLevel, GWarn, AccessTokenString, TokenExpiresAt, bWasInteractiveLogin))
 	{
 		if (bWasInteractiveLogin)
 		{
@@ -756,7 +756,7 @@ FBuildServiceInstance::SetAccessTokenAndUnlock(FScopeLock& Lock, FStringView Tok
 			RefreshAccessTokenHandle = FTSTicker::GetCoreTicker().AddTicker(FTickerDelegate::CreateLambda(
 				[this](float DeltaTime)
 				{
-					AcquireAccessToken();
+					AcquireAccessToken(EDesktopLoginInteractionLevel::None);
 					return false;
 				}
 			), float(FMath::Min(RefreshDelay - RefreshGracePeriod, MAX_flt)));
diff --git a/Engine/Source/Developer/Zen/Public/Experimental/BuildServerInterface.h b/Engine/Source/Developer/Zen/Public/Experimental/BuildServerInterface.h
--- a/Engine/Source/Developer/Zen/Public/Experimental/BuildServerInterface.h
+++ b/Engine/Source/Developer/Zen/Public/Experimental/BuildServerInterface.h
@@ -21,6 +21,9 @@
 
 #define UE_API ZEN_API
 
+
+enum class EDesktopLoginInteractionLevel : uint8;
+
 namespace UE { class IHttpClient; }
 namespace UE::Zen { class FScopeZenService; }
 namespace UE::Zen::Build::Private { class FAccessToken; }
@@ -65,6 +68,15 @@ public:
 		ConnectionFailed
 	};
 
+	enum class EConnectionFailureReason
+	{
+		FailedToResolveHost,
+		FailedToAcquireAccessToken,
+		UnexpectedState,
+		Unknown,
+		None
+	};
+
 	struct FBuildRecord
 	{
 		FCbObjectId BuildId;
@@ -100,7 +112,7 @@ public:
 		friend class UE::Zen::Build::FBuildServiceInstance;
 	};
 
-	typedef TUniqueFunction<void(EConnectionState)> FOnConnectionComplete;
+	typedef TUniqueFunction<void(EConnectionState, EConnectionFailureReason)> FOnConnectionComplete;
 	typedef TUniqueFunction<void()> FOnRefreshNamespacesAndBucketsComplete;
 	typedef TUniqueFunction<void(TArray<FBuildRecord>&& ListBuildRecords)> FOnListBuildsComplete;
 
@@ -110,7 +122,7 @@ public:
 
 	UE_API virtual ~FBuildServiceInstance();
 
-	UE_API void Connect(FOnConnectionComplete&& OnConnectionComplete = {});
+	UE_API void Connect(bool bAllowInteractive, FOnConnectionComplete&& OnConnectionComplete = {});
 	UE_API void RefreshNamespacesAndBuckets(FOnRefreshNamespacesAndBucketsComplete&& InOnRefreshNamespacesAndBucketsComplete = {});
 	[[nodiscard]] UE_API EConnectionState GetConnectionState() const;
 	[[nodiscard]] UE_API FAnsiStringView GetEffectiveDomain() const;
@@ -129,7 +141,7 @@ private:
 	typedef TUniqueFunction<void(bool)> FOnZenUtilityInvocationComplete;
 
 	void Initialize();
-	bool AcquireAccessToken();
+	bool AcquireAccessToken(EDesktopLoginInteractionLevel InteractionLevel);
 	void SetAccessTokenAndUnlock(FScopeLock &Lock, FStringView Token, double RefreshDelay = 0.0);
 	[[nodiscard]] FString GetAccessToken() const;
 	FString AddZenUtilityBuildServerArguments(const TCHAR* InCommandLineArgs) const;
diff --git a/Engine/Source/Editor/StorageServerWidgets/Private/SBuildLogin.cpp b/Engine/Source/Editor/StorageServerWidgets/Private/SBuildLogin.cpp
--- a/Engine/Source/Editor/StorageServerWidgets/Private/SBuildLogin.cpp
+++ b/Engine/Source/Editor/StorageServerWidgets/Private/SBuildLogin.cpp
@@ -34,6 +34,19 @@ void SBuildLogin::Construct(const FArguments& InArgs)
 			GetGridPanel()
 		]
 	];
+
+	if (TSharedPtr<UE::Zen::Build::FBuildServiceInstance> ServiceInstance = BuildServiceInstance.Get())
+	{
+		ServiceInstance->Connect(false, [ServiceInstance]
+			(UE::Zen::Build::FBuildServiceInstance::EConnectionState ConnectionState,
+				UE::Zen::Build::FBuildServiceInstance::EConnectionFailureReason FailureReason)
+			{
+				if (ConnectionState == UE::Zen::Build::FBuildServiceInstance::EConnectionState::ConnectionSucceeded)
+				{
+					ServiceInstance->RefreshNamespacesAndBuckets();
+				}
+			});
+	}
 }
 
 TSharedRef<SWidget> SBuildLogin::GetGridPanel()
@@ -55,7 +68,7 @@ TSharedRef<SWidget> SBuildLogin::GetGridPanel()
 	];
 
 	Panel->AddSlot(1, 0)
-	.Padding(FMargin(ColumnMargin, RowMargin))
+	.Padding(FMargin(2.0f, RowMargin))
 	[
 		SNew(STextBlock)
 		.Text_Lambda([this]
@@ -93,39 +106,69 @@ TSharedRef<SWidget> SBuildLogin::GetGridPanel()
 	];
 
 	Panel->AddSlot(2, 0)
-	.Padding(FMargin(ColumnMargin, RowMargin))
+	.Padding(FMargin(2.0f, RowMargin))
 	[
-		SNew(SHyperlink)
-		.Style(FAppStyle::Get(), TEXT("NavigationHyperlink"))
-		.Text_Lambda([this]
-		{
-			FText BrowseText = LOCTEXT("BuildLogin_ConnectLink", "Click to connect");
-			return BrowseText;
-		})
-		.OnNavigate_Lambda([this]
-		{
-			if (TSharedPtr<UE::Zen::Build::FBuildServiceInstance> ServiceInstance = BuildServiceInstance.Get())
-			{
-				ServiceInstance->Connect([ServiceInstance](UE::Zen::Build::FBuildServiceInstance::EConnectionState ConnectionState)
-					{
-						if (ConnectionState == UE::Zen::Build::FBuildServiceInstance::EConnectionState::ConnectionSucceeded)
-						{
-							ServiceInstance->RefreshNamespacesAndBuckets();
-						}
-					});
-			}
-		})
+		SNew(SHorizontalBox)
 		.Visibility_Lambda([this]
 		{
 			if (TSharedPtr<UE::Zen::Build::FBuildServiceInstance> ServiceInstance = BuildServiceInstance.Get())
 			{
-				if (ServiceInstance->GetConnectionState() == UE::Zen::Build::FBuildServiceInstance::EConnectionState::NotStarted)
+				if (ServiceInstance->GetConnectionState() != UE::Zen::Build::FBuildServiceInstance::EConnectionState::ConnectionInProgress)
 				{
 					return EVisibility::Visible;
 				}
 			}
 			return EVisibility::Hidden;
 		})
+		+SHorizontalBox::Slot()
+		.AutoWidth()
+		.VAlign(VAlign_Center)
+		[
+			SNew(STextBlock)
+			.Text(FText::FromString(TEXT("(")))
+		]
+		+SHorizontalBox::Slot()
+		.AutoWidth()
+		.VAlign(VAlign_Center)
+		[
+			SNew(SHyperlink)
+			.Style(FAppStyle::Get(), TEXT("NavigationHyperlink"))
+			.Text_Lambda([this]
+			{
+				if (TSharedPtr<UE::Zen::Build::FBuildServiceInstance> ServiceInstance = BuildServiceInstance.Get())
+				{
+					switch (ServiceInstance->GetConnectionState())
+					{
+					case UE::Zen::Build::FBuildServiceInstance::EConnectionState::ConnectionSucceeded:
+					case UE::Zen::Build::FBuildServiceInstance::EConnectionState::ConnectionFailed:
+						return LOCTEXT("BuildLogin_ReconnectLink", "reconnect");
+					}
+				}
+				return LOCTEXT("BuildLogin_ConnectLink", "Click to connect");
+			})
+			.OnNavigate_Lambda([this]
+			{
+				if (TSharedPtr<UE::Zen::Build::FBuildServiceInstance> ServiceInstance = BuildServiceInstance.Get())
+				{
+					ServiceInstance->Connect(true, [ServiceInstance]
+						(UE::Zen::Build::FBuildServiceInstance::EConnectionState ConnectionState,
+							UE::Zen::Build::FBuildServiceInstance::EConnectionFailureReason FailureReason)
+						{
+							if (ConnectionState == UE::Zen::Build::FBuildServiceInstance::EConnectionState::ConnectionSucceeded)
+							{
+								ServiceInstance->RefreshNamespacesAndBuckets();
+							}
+						});
+				}
+			})
+		]
+		+SHorizontalBox::Slot()
+		.AutoWidth()
+		.VAlign(VAlign_Center)
+		[
+			SNew(STextBlock)
+			.Text(FText::FromString(TEXT(")")))
+		]
 	];
 
 
