Handle issues related to build server interface when no host has been configured.  Make sure that OidcToken acquisition doesn't fail on a 'check' statement because the OidcToken executable threw an exception due (which doesn't yield a return code from the subprocess on Mac & Linux).  Also ensure that the higher level code for the build server interface doesn't try to connect if the configured host is invalid.

diff --git a/Engine/Source/Developer/DesktopPlatform/Private/DesktopPlatformBase.cpp b/Engine/Source/Developer/DesktopPlatform/Private/DesktopPlatformBase.cpp
--- a/Engine/Source/Developer/DesktopPlatform/Private/DesktopPlatformBase.cpp
+++ b/Engine/Source/Developer/DesktopPlatform/Private/DesktopPlatformBase.cpp
@@ -694,6 +694,10 @@ bool FDesktopPlatformBase::InvokeUnrealBuildToolSync(const FString& InCmdLinePar
 
 
 	FPlatformProcess::ClosePipe(PipeRead, PipeWrite);
+	if (ProcHandle.IsValid())
+	{
+		FPlatformProcess::CloseProc(ProcHandle);
+	}
 
 	return bInvoked;
 }
@@ -1236,12 +1240,26 @@ bool FDesktopPlatformBase::GetOidcAccessTokenInteractive(const FString& RootDir,
 		OutReturnCode = -1;
 		FPlatformProcess::ClosePipe(PipeRead, PipeWrite);
 		FPlatformProcess::TerminateProc(ProcHandle);
+		if (ProcHandle.IsValid())
+		{
+			FPlatformProcess::CloseProc(ProcHandle);
+		}
+
 		return false;
 	}
 
-	bool bGotReturnCode = FPlatformProcess::GetProcReturnCode(ProcHandle, &OutReturnCode);		
-	check(bGotReturnCode);
+	if (!FPlatformProcess::GetProcReturnCode(ProcHandle, &OutReturnCode))
+	{
+		// OidcToken can throw exceptions, and in those cases on Mac & Linux, we don't
+		// get a return code.  Windows always produces a return code.  Indicate it was invoked
+		// but yield a return code of -1.
+		OutReturnCode = -1;
+	}
 	FPlatformProcess::ClosePipe(PipeRead, PipeWrite);
+	if (ProcHandle.IsValid())
+	{
+		FPlatformProcess::CloseProc(ProcHandle);
+	}
 
 	return bInvoked;
 }
@@ -1267,8 +1285,13 @@ bool FDesktopPlatformBase::InvokeOidcTokenToolSync(const FText& Description, con
 			FPlatformProcess::Sleep(0.1f);
 		}		
 		bInvoked = true;
-		bool bGotReturnCode = FPlatformProcess::GetProcReturnCode(ProcHandle, &OutReturnCode);		
-		check(bGotReturnCode);
+		if (!FPlatformProcess::GetProcReturnCode(ProcHandle, &OutReturnCode))
+		{
+			// OidcToken can throw exceptions, and in those cases on Mac & Linux, we don't
+			// get a return code.  Windows always produces a return code.  Indicate it was invoked
+			// but yield a return code of -1.
+			OutReturnCode = -1;
+		}
 	}
 	else
 	{
@@ -1279,6 +1302,10 @@ bool FDesktopPlatformBase::InvokeOidcTokenToolSync(const FText& Description, con
 
 
 	FPlatformProcess::ClosePipe(PipeRead, PipeWrite);
+	if (ProcHandle.IsValid())
+	{
+		FPlatformProcess::CloseProc(ProcHandle);
+	}
 
 	return bInvoked;
 }
diff --git a/Engine/Source/Developer/Zen/Private/BuildServerInterface.cpp b/Engine/Source/Developer/Zen/Private/BuildServerInterface.cpp
--- a/Engine/Source/Developer/Zen/Private/BuildServerInterface.cpp
+++ b/Engine/Source/Developer/Zen/Private/BuildServerInterface.cpp
@@ -224,6 +224,16 @@ FBuildServiceInstance::~FBuildServiceInstance()
 void
 FBuildServiceInstance::Connect(bool bAllowInteractive, FOnConnectionComplete&& OnConnectionComplete)
 {
+	const FString Host = Settings.GetHost();
+	if (Host.IsEmpty() || (Host == TEXT("None")))
+	{
+		if (OnConnectionComplete)
+		{
+			OnConnectionComplete(EConnectionState::ConnectionFailed, EConnectionFailureReason::MissingHost);
+		}
+		return;
+	}
+
 	if (ConnectionState.exchange(EConnectionState::ConnectionInProgress, std::memory_order_relaxed) == EConnectionState::ConnectionInProgress)
 	{
 		if (OnConnectionComplete)
diff --git a/Engine/Source/Developer/Zen/Public/Experimental/BuildServerInterface.h b/Engine/Source/Developer/Zen/Public/Experimental/BuildServerInterface.h
--- a/Engine/Source/Developer/Zen/Public/Experimental/BuildServerInterface.h
+++ b/Engine/Source/Developer/Zen/Public/Experimental/BuildServerInterface.h
@@ -70,6 +70,7 @@ public:
 
 	enum class EConnectionFailureReason
 	{
+		MissingHost,
 		FailedToResolveHost,
 		FailedToAcquireAccessToken,
 		UnexpectedState,
