PreviewPlatforms - Fix hang on macOS when using "Preview via Json..."
This is a fix for an issue caused by not using a pair when defining filetypes.
Fix the workaround when only the extension is provided by duplicating the extension for description.
The issue was caused by UTType::typeWithFilenameExtension returning nil with the previous workaround.

#tests local test + preflights

#Synced-CL Unknown

diff --git a/Engine/Source/Developer/DesktopPlatform/Private/Mac/DesktopPlatformMac.cpp b/Engine/Source/Developer/DesktopPlatform/Private/Mac/DesktopPlatformMac.cpp
--- a/Engine/Source/Developer/DesktopPlatform/Private/Mac/DesktopPlatformMac.cpp
+++ b/Engine/Source/Developer/DesktopPlatform/Private/Mac/DesktopPlatformMac.cpp
@@ -428,7 +428,8 @@ bool FDesktopPlatformMac::FileDialogShared(bool bSave, const void* ParentWindowH
 
 			if ([AllowedFileTypes count] == 1)
 			{
-				[AllowedFileTypes addObject:@""];
+				// Duplicate the type for the description
+				[AllowedFileTypes addObject:AllowedFileTypes[0]];
 			}
 			
 			[AccessoryView AddAllowedFileTypes:AllowedFileTypes];
diff --git a/Engine/Source/Editor/LevelEditor/Private/LevelEditorActions.cpp b/Engine/Source/Editor/LevelEditor/Private/LevelEditorActions.cpp
--- a/Engine/Source/Editor/LevelEditor/Private/LevelEditorActions.cpp
+++ b/Engine/Source/Editor/LevelEditor/Private/LevelEditorActions.cpp
@@ -564,7 +564,7 @@ void FLevelEditorActionCallbacks::PreviewJson_Clicked(FName PlatformName, FName
 				NSLOCTEXT("UnrealEd", "PreviewJson", "Preview Json").ToString(),
 				DefaultLocation,
 				TEXT(""),
-				TEXT("*.json"),
+				TEXT("json (*.json)|*.json"),
 				EFileDialogFlags::None,
 				OpenedFiles
 			);
