Handle issues related to build server interface when no host has been configured.  Make sure that OidcToken acquisition doesn't fail on a 'check' statement because the OidcToken executable threw an exception due (which doesn't yield a return code from the subprocess on Mac & Linux).  Also ensure that the higher level code for the build server interface doesn't try to connect if the configured host is invalid.

diff --git a/Engine/Source/Developer/DesktopPlatform/Private/DesktopPlatformBase.cpp b/Engine/Source/Developer/DesktopPlatform/Private/DesktopPlatformBase.cpp
--- a/Engine/Source/Developer/DesktopPlatform/Private/DesktopPlatformBase.cpp
+++ b/Engine/Source/Developer/DesktopPlatform/Private/DesktopPlatformBase.cpp
@@ -694,6 +694,10 @@ bool FDesktopPlatformBase::InvokeUnrealBuildToolSync(const FString& InCmdLinePar
 
 
 	FPlatformProcess::ClosePipe(PipeRead, PipeWrite);
+	if (ProcHandle.IsValid())
+	{
+		FPlatformProcess::CloseProc(ProcHandle);
+	}
 
 	return bInvoked;
 }
@@ -1236,6 +1240,11 @@ bool FDesktopPlatformBase::GetOidcAccessTokenInteractive(const FString& RootDir,
 		OutReturnCode = -1;
 		FPlatformProcess::ClosePipe(PipeRead, PipeWrite);
 		FPlatformProcess::TerminateProc(ProcHandle);
+		if (ProcHandle.IsValid())
+		{
+			FPlatformProcess::CloseProc(ProcHandle);
+		}
+
 		return false;
 	}
 
@@ -1247,6 +1256,10 @@ bool FDesktopPlatformBase::GetOidcAccessTokenInteractive(const FString& RootDir,
 		OutReturnCode = -1;
 	}
 	FPlatformProcess::ClosePipe(PipeRead, PipeWrite);
+	if (ProcHandle.IsValid())
+	{
+		FPlatformProcess::CloseProc(ProcHandle);
+	}
 
 	return bInvoked;
 }
@@ -1289,6 +1302,10 @@ bool FDesktopPlatformBase::InvokeOidcTokenToolSync(const FText& Description, con
 
 
 	FPlatformProcess::ClosePipe(PipeRead, PipeWrite);
+	if (ProcHandle.IsValid())
+	{
+		FPlatformProcess::CloseProc(ProcHandle);
+	}
 
 	return bInvoked;
 }
