Fixed a case senstive issue when processing targets in GetTargetsForProject that was causing the editor to take a long time to boot up.

#rnx
#Synced-CL 48524062

diff --git a/Engine/Source/Developer/DesktopPlatform/Private/DesktopPlatformBase.cpp b/Engine/Source/Developer/DesktopPlatform/Private/DesktopPlatformBase.cpp
--- a/Engine/Source/Developer/DesktopPlatform/Private/DesktopPlatformBase.cpp
+++ b/Engine/Source/Developer/DesktopPlatform/Private/DesktopPlatformBase.cpp
@@ -1357,7 +1357,7 @@ struct FTargetFileVisitor : IPlatformFile::FDirectoryStatVisitor
 		FStringView FileNameOrDirectory(InFileNameOrDirectory);
 
 		// NOTE: This code needs to behave the same as FindAllRulesSourceFiles in Rules.cs
-		if (FStringView WithoutExtension = UE::String::RemoveFromEnd(FileNameOrDirectory, TEXTVIEW(".Target.cs"));
+		if (FStringView WithoutExtension = UE::String::RemoveFromEnd(FileNameOrDirectory, TEXTVIEW(".Target.cs"), ESearchCase::IgnoreCase);
 			WithoutExtension.Len() != FileNameOrDirectory.Len())
 		{
 			FString TargetName = FPaths::GetCleanFilename(FString(WithoutExtension));
