DDB: Reworked the build worker API to prepare for build worker packages

- Represent build workers as compact binary with a "key" that is the hash of that representation.
- Move resolving of build worker files into a new interface IBuildWorkerResolver.
- Move creation of build workers into the factory to allow returning a worker not created by the provided builder.
- Make FBuildWorker a concrete class that can be moved and has intrusive optional support.
- Make FBuildWorkerBuilder a concrete class with a workaround to allow use with only an include dependency.
- Split up FTargetReceiptBuildWorker to prepare for loading target receipts from outside of the factory.
- Minor cleanup to FMultiPlatformTargetReceiptBuildWorker.

#Synced-CL 50253651

diff --git a/Engine/Source/Developer/DerivedDataCache/Private/DerivedDataBuild.cpp b/Engine/Source/Developer/DerivedDataCache/Private/DerivedDataBuild.cpp
--- a/Engine/Source/Developer/DerivedDataCache/Private/DerivedDataBuild.cpp
+++ b/Engine/Source/Developer/DerivedDataCache/Private/DerivedDataBuild.cpp
@@ -19,11 +19,16 @@
 #include "DerivedDataPrivate.h"
 #include "Misc/Guid.h"
 
-namespace UE::DerivedData::Private
+namespace UE::DerivedData
 {
 
 DEFINE_LOG_CATEGORY(LogDerivedDataBuild);
 
+} // UE::DerivedData
+
+namespace UE::DerivedData::Private
+{
+
 /**
  * Derived Data Build System
  *
diff --git a/Engine/Source/Developer/DerivedDataCache/Private/DerivedDataBuildPrivate.h b/Engine/Source/Developer/DerivedDataCache/Private/DerivedDataBuildPrivate.h
--- a/Engine/Source/Developer/DerivedDataCache/Private/DerivedDataBuildPrivate.h
+++ b/Engine/Source/Developer/DerivedDataCache/Private/DerivedDataBuildPrivate.h
@@ -29,11 +29,16 @@ namespace UE::DerivedData { class IRequestOwner; }
 namespace UE::DerivedData { struct FBuildKey; }
 namespace UE::DerivedData { struct FValueId; }
 
-namespace UE::DerivedData::Private
+namespace UE::DerivedData
 {
 
 DECLARE_LOG_CATEGORY_EXTERN(LogDerivedDataBuild, Log, All);
 
+} // UE::DerivedData
+
+namespace UE::DerivedData::Private
+{
+
 // Implemented in DerivedDataBuild.cpp
 IBuild* CreateBuild(ICache* Cache);
 
diff --git a/Engine/Source/Developer/DerivedDataCache/Private/DerivedDataBuildWorkerRegistry.cpp b/Engine/Source/Developer/DerivedDataCache/Private/DerivedDataBuildWorkerRegistry.cpp
--- a/Engine/Source/Developer/DerivedDataCache/Private/DerivedDataBuildWorkerRegistry.cpp
+++ b/Engine/Source/Developer/DerivedDataCache/Private/DerivedDataBuildWorkerRegistry.cpp
@@ -3,138 +3,185 @@
 #include "DerivedDataBuildWorkerRegistry.h"
 
 #include "Algo/Find.h"
+#include "Async/SharedLock.h"
+#include "Async/SharedMutex.h"
+#include "Async/UniqueLock.h"
 #include "Containers/SharedString.h"
 #include "Containers/StringView.h"
-#include "Containers/UnrealString.h"
 #include "DerivedDataBuildKey.h"
 #include "DerivedDataBuildPrivate.h"
 #include "DerivedDataBuildWorker.h"
 #include "Features/IModularFeatures.h"
-#include "HAL/CriticalSection.h"
-#include "IO/IoHash.h"
-#include "Misc/Guid.h"
-#include "Misc/ScopeRWLock.h"
+#include "Serialization/CompactBinaryWriter.h"
 #include "Templates/Function.h"
-#include "Templates/Tuple.h"
 
-namespace UE::DerivedData::Private
+namespace UE::DerivedData
 {
 
-class FBuildWorkerInternal final : public FBuildWorker, public FBuildWorkerBuilder
+FBuildWorker::FBuildWorker(const FBuildWorkerKey& InKey, FCbObject&& InWorker, TUniquePtr<IBuildWorkerResolver> InResolver)
+	: Worker(MoveTemp(InWorker))
+	, Key(InKey)
+	, Resolver(MoveTemp(InResolver))
 {
-public:
-	inline explicit FBuildWorkerInternal(IBuildWorkerFactory* InFactory)
-		: Factory(InFactory)
-	{
-	}
-
-	void Build();
-
-	inline FStringView GetName() const final { return WorkerName; }
-	inline FStringView GetPath() const final { return WorkerPath; }
-	inline FStringView GetHostPlatform() const final { return HostPlatform; }
-	inline FGuid GetBuildSystemVersion() const final { return BuildSystemVersion; }
-
-	void FindFileData(TConstArrayView<FIoHash> RawHashes, IRequestOwner& Owner, FOnBuildWorkerFileDataComplete&& OnComplete) const final;
-
-	void IterateFunctions(TFunctionRef<void (FUtf8StringView Name, const FGuid& Version)> Visitor) const final;
-	void IterateFiles(TFunctionRef<void (FStringView Path, const FIoHash& RawHash, uint64 RawSize)> Visitor) const final;
-	void IterateExecutables(TFunctionRef<void (FStringView Path, const FIoHash& RawHash, uint64 RawSize)> Visitor) const final;
-	void IterateEnvironment(TFunctionRef<void (FStringView Name, FStringView Value)> Visitor) const final;
+}
 
-private:
-	inline void SetName(FStringView Name) final { WorkerName = Name; }
-	inline void SetPath(FStringView Path) final { WorkerPath = Path; }
-	inline void SetHostPlatform(FStringView Name) final { HostPlatform = Name; }
-	inline void SetBuildSystemVersion(const FGuid& Version) final { BuildSystemVersion = Version; }
+FBuildWorkerKey FBuildWorker::GetKey() const
+{
+	return Key;
+}
 
-	void AddFunction(FUtf8StringView Name, const FGuid& Version) final;
-	void AddFile(FStringView Path, const FIoHash& RawHash, uint64 RawSize) final;
-	void AddExecutable(FStringView Path, const FIoHash& RawHash, uint64 RawSize) final;
-	void SetEnvironment(FStringView Name, FStringView Value) final;
+FUtf8StringView FBuildWorker::GetName() const
+{
+	return Worker.FindView(ANSITEXTVIEW("Name")).AsString();
+}
 
-private:
-	FString WorkerName;
-	FString WorkerPath;
-	FString HostPlatform;
-	FGuid BuildSystemVersion;
-	TArray<TTuple<FUtf8SharedString, FGuid>> Functions;
-	TArray<TTuple<FString, FIoHash, uint64>> Files;
-	TArray<TTuple<FString, FIoHash, uint64>> Executables;
-	TArray<TTuple<FString, FString>> Environment;
-	IBuildWorkerFactory* Factory;
-};
+FUtf8StringView FBuildWorker::GetPath() const
+{
+	return Worker.FindView(ANSITEXTVIEW("Path")).AsString();
+}
 
-void FBuildWorkerInternal::Build()
+FUtf8StringView FBuildWorker::GetHostPlatform() const
 {
-	Functions.Sort();
-	Files.Sort();
-	Executables.Sort();
-	Environment.Sort();
+	return Worker.FindView(ANSITEXTVIEW("HostPlatform")).AsString();
 }
 
-void FBuildWorkerInternal::FindFileData(TConstArrayView<FIoHash> RawHashes, IRequestOwner& Owner, FOnBuildWorkerFileDataComplete&& OnComplete) const
+FGuid FBuildWorker::GetBuildSystemVersion() const
 {
-	return Factory->FindFileData(RawHashes, Owner, MoveTemp(OnComplete));
+	return Worker.FindView(ANSITEXTVIEW("BuildSystemVersion")).AsUuid();
 }
 
-void FBuildWorkerInternal::IterateFunctions(TFunctionRef<void (FUtf8StringView Name, const FGuid& Version)> Visitor) const
+void FBuildWorker::IterateFunctions(TFunctionRef<void (FUtf8StringView Name, const FGuid& Version)> Visitor) const
 {
-	for (const TTuple<FUtf8SharedString, FGuid>& Function : Functions)
+	for (FCbFieldView It : Worker.FindView(ANSITEXTVIEW("Functions")))
 	{
-		Function.ApplyAfter(Visitor);
+		Visitor(It[ANSITEXTVIEW("Name")].AsString(), It[ANSITEXTVIEW("Version")].AsUuid());
 	}
 }
 
-void FBuildWorkerInternal::IterateFiles(TFunctionRef<void (FStringView Path, const FIoHash& RawHash, uint64 RawSize)> Visitor) const
+void FBuildWorker::IterateFiles(TFunctionRef<void (FUtf8StringView Path, const FIoHash& RawHash, uint64 RawSize)> Visitor) const
 {
-	for (const TTuple<FString, FIoHash, uint64>& File : Files)
+	for (FCbFieldView It : Worker.FindView(ANSITEXTVIEW("Files")))
 	{
-		File.ApplyAfter(Visitor);
+		Visitor(It[ANSITEXTVIEW("Path")].AsString(), It[ANSITEXTVIEW("RawHash")].AsHash(), It[ANSITEXTVIEW("RawSize")].AsUInt64());
 	}
 }
 
-void FBuildWorkerInternal::IterateExecutables(TFunctionRef<void (FStringView Path, const FIoHash& RawHash, uint64 RawSize)> Visitor) const
+void FBuildWorker::IterateExecutables(TFunctionRef<void (FUtf8StringView Path, const FIoHash& RawHash, uint64 RawSize)> Visitor) const
 {
-	for (const TTuple<FString, FIoHash, uint64>& Executable : Executables)
+	for (FCbFieldView It : Worker.FindView(ANSITEXTVIEW("Executables")))
 	{
-		Executable.ApplyAfter(Visitor);
+		Visitor(It[ANSITEXTVIEW("Path")].AsString(), It[ANSITEXTVIEW("RawHash")].AsHash(), It[ANSITEXTVIEW("RawSize")].AsUInt64());
 	}
 }
 
-void FBuildWorkerInternal::IterateEnvironment(TFunctionRef<void (FStringView Name, FStringView Value)> Visitor) const
+void FBuildWorker::IterateEnvironment(TFunctionRef<void (FUtf8StringView Name, FUtf8StringView Value)> Visitor) const
 {
-	for (const TTuple<FString, FString>& Variable : Environment)
+	for (FCbFieldView It : Worker.FindView(ANSITEXTVIEW("Environment")))
 	{
-		Variable.ApplyAfter(Visitor);
+		Visitor(It[ANSITEXTVIEW("Name")].AsString(), It[ANSITEXTVIEW("Value")].AsString());
 	}
 }
 
-void FBuildWorkerInternal::AddFunction(FUtf8StringView Name, const FGuid& Version)
+void FBuildWorker::Resolve(TConstArrayView<FIoHash> RawHashes, IRequestOwner& Owner, FOnBuildWorkerResolved&& OnResolved) const
 {
-	UE_CLOG(!Version.IsValid(), LogDerivedDataBuild, Error,
-		TEXT("Version of zero is not allowed in build function with the name %s in build worker '%s'."),
-		*WriteToString<32>(Name), *WorkerName);
-	Functions.Emplace(Name, Version);
+	Resolver->Resolve(RawHashes, Owner, MoveTemp(OnResolved));
 }
 
-void FBuildWorkerInternal::AddFile(FStringView Path, const FIoHash& RawHash, uint64 RawSize)
+TOptional<FBuildWorker> FBuildWorker::Load(const FBuildWorkerKey& Key, FCbObject&& Worker, TUniquePtr<IBuildWorkerResolver> Resolver)
 {
-	Files.Emplace(Path, RawHash, RawSize);
+	if (Key.Hash == Worker.GetHash())
+	{
+		return FBuildWorker(Key, MoveTemp(Worker), MoveTemp(Resolver));
+	}
+	return {};
 }
 
-void FBuildWorkerInternal::AddExecutable(FStringView Path, const FIoHash& RawHash, uint64 RawSize)
+FBuildWorker FBuildWorkerBuilder::Build(TUniquePtr<IBuildWorkerResolver> Resolver)
 {
-	Executables.Emplace(Path, RawHash, RawSize);
-}
+	Functions.Sort();
+	Files.Sort();
+	Executables.Sort();
+	Environment.Sort();
 
-void FBuildWorkerInternal::SetEnvironment(FStringView Name, FStringView Value)
-{
-	Environment.Emplace(Name, Value);
+	TCbWriter<4096> Writer;
+	Writer.BeginObject();
+
+	Writer.AddString(ANSITEXTVIEW("Name"), WorkerName);
+	Writer.AddString(ANSITEXTVIEW("Path"), WorkerPath);
+	Writer.AddString(ANSITEXTVIEW("HostPlatform"), HostPlatform);
+	Writer.AddUuid(ANSITEXTVIEW("BuildSystemVersion"), BuildSystemVersion);
+
+	if (!Functions.IsEmpty())
+	{
+		Writer.BeginArray(ANSITEXTVIEW("Functions"));
+		for (const TTuple<FUtf8String, FGuid>& Function : Functions)
+		{
+			const FUtf8String& Name = Function.Get<FUtf8String>();
+			const FGuid& Version = Function.Get<FGuid>();
+			UE_CLOGF(!Version.IsValid(), LogDerivedDataBuild, Error,
+				"Version of zero is not allowed in build function with the name %s in build worker '%s'.",
+				*WriteToUtf8String<32>(Name), *WorkerName);
+			Writer.BeginObject();
+			Writer.AddString(ANSITEXTVIEW("Name"), Name);
+			Writer.AddUuid(ANSITEXTVIEW("Version"), Version);
+			Writer.EndObject();
+		}
+		Writer.EndArray();
+	}
+
+	if (!Files.IsEmpty())
+	{
+		Writer.BeginArray(ANSITEXTVIEW("Files"));
+		for (const TTuple<FUtf8String, FIoHash, uint64>& File : Files)
+		{
+			Writer.BeginObject();
+			Writer.AddString(ANSITEXTVIEW("Path"), File.Get<FUtf8String>());
+			Writer.AddBinaryAttachment(ANSITEXTVIEW("RawHash"), File.Get<FIoHash>());
+			Writer.AddInteger(ANSITEXTVIEW("RawSize"), File.Get<uint64>());
+			Writer.EndObject();
+		}
+		Writer.EndArray();
+	}
+
+	if (!Executables.IsEmpty())
+	{
+		Writer.BeginArray(ANSITEXTVIEW("Executables"));
+		for (const TTuple<FUtf8String, FIoHash, uint64>& Executable : Executables)
+		{
+			Writer.BeginObject();
+			Writer.AddString(ANSITEXTVIEW("Path"), Executable.Get<FUtf8String>());
+			Writer.AddBinaryAttachment(ANSITEXTVIEW("RawHash"), Executable.Get<FIoHash>());
+			Writer.AddInteger(ANSITEXTVIEW("RawSize"), Executable.Get<uint64>());
+			Writer.EndObject();
+		}
+		Writer.EndArray();
+	}
+
+	if (!Environment.IsEmpty())
+	{
+		Writer.BeginArray(ANSITEXTVIEW("Environment"));
+		for (const TTuple<FUtf8String, FUtf8String>& Var : Environment)
+		{
+			Writer.BeginObject();
+			Writer.AddString(ANSITEXTVIEW("Name"), Var.Get<0>());
+			Writer.AddString(ANSITEXTVIEW("Value"), Var.Get<1>());
+			Writer.EndObject();
+		}
+		Writer.EndArray();
+	}
+
+	Writer.EndObject();
+	FCbObject Worker = Writer.Save().AsObject();
+	return FBuildWorker({Worker.GetHash()}, MoveTemp(Worker), MoveTemp(Resolver));
 }
 
+} // UE::DerivedData
+
 ///////////////////////////////////////////////////////////////////////////////////////////////////////////////////////
 
+namespace UE::DerivedData::Private
+{
+
 class FBuildWorkerRegistry final : public IBuildWorkerRegistry
 {
 public:
@@ -155,10 +202,10 @@ private:
 	void RemoveWorker(IBuildWorkerFactory* Factory);
 
 private:
-	mutable FRWLock Lock;
 	IBuildWorkerExecutor* Executor = nullptr;
 	TMap<IBuildWorkerFactory*, TUniquePtr<FBuildWorker>> Workers;
 	TMultiMap<TTuple<FUtf8SharedString, FGuid>, FBuildWorker*> Functions;
+	mutable FSharedMutex Mutex;
 };
 
 FBuildWorkerRegistry::FBuildWorkerRegistry()
@@ -187,7 +234,7 @@ void FBuildWorkerRegistry::OnModularFeatureRegistered(const FName& Type, IModula
 {
 	if (!Executor && Type == IBuildWorkerExecutor::FeatureName)
 	{
-		FWriteScopeLock WriteLock(Lock);
+		TUniqueLock Lock(Mutex);
 		Executor = static_cast<IBuildWorkerExecutor*>(ModularFeature);
 	}
 	else if (Type == IBuildWorkerFactory::FeatureName)
@@ -206,7 +253,7 @@ void FBuildWorkerRegistry::OnModularFeatureUnregistered(const FName& Type, IModu
 		{
 			NextExecutor = &ModularFeatures.GetModularFeature<IBuildWorkerExecutor>(IBuildWorkerExecutor::FeatureName);
 		}
-		FWriteScopeLock WriteLock(Lock);
+		TUniqueLock Lock(Mutex);
 		Executor = static_cast<IBuildWorkerExecutor*>(NextExecutor);
 	}
 	else if (Type == IBuildWorkerFactory::FeatureName)
@@ -217,21 +264,22 @@ void FBuildWorkerRegistry::OnModularFeatureUnregistered(const FName& Type, IModu
 
 void FBuildWorkerRegistry::AddWorker(IBuildWorkerFactory* Factory)
 {
-	TUniquePtr<FBuildWorkerInternal> Worker = MakeUnique<FBuildWorkerInternal>(Factory);
-	Factory->Build(*Worker);
-	Worker->Build();
-
-	FWriteScopeLock WriteLock(Lock);
-	Worker->IterateFunctions([this, Worker = Worker.Get()](FUtf8StringView Name, const FGuid& Version)
+	FBuildWorkerBuilder Builder;
+	if (TOptional<FBuildWorker> OptionalWorker = Factory->Build(Builder))
 	{
-		Functions.Emplace(MakeTuple(FUtf8SharedString(Name), Version), Worker);
-	});
-	Workers.Emplace(Factory, MoveTemp(Worker));
+		TUniquePtr<FBuildWorker> Worker = MakeUnique<FBuildWorker>(MoveTemp(OptionalWorker.GetValue()));
+		TUniqueLock Lock(Mutex);
+		Worker->IterateFunctions([this, Worker = Worker.Get()](FUtf8StringView Name, const FGuid& Version)
+		{
+			Functions.Emplace(MakeTuple(FUtf8SharedString(Name), Version), Worker);
+		});
+		Workers.Emplace(Factory, MoveTemp(Worker));
+	}
 }
 
 void FBuildWorkerRegistry::RemoveWorker(IBuildWorkerFactory* Factory)
 {
-	FWriteScopeLock WriteLock(Lock);
+	TUniqueLock Lock(Mutex);
 	TUniquePtr<FBuildWorker>& Worker = Workers.FindChecked(Factory);
 	Worker->IterateFunctions([this, Worker = Worker.Get()](FUtf8StringView Name, const FGuid& Version)
 	{
@@ -246,10 +294,10 @@ FBuildWorker* FBuildWorkerRegistry::FindWorker(
 	const FGuid& BuildSystemVersion,
 	IBuildWorkerExecutor*& OutWorkerExecutor) const
 {
-	FReadScopeLock ReadLock(Lock);
+	TSharedLock Lock(Mutex);
 	if (Executor)
 	{
-		TConstArrayView<FStringView> ExecutorHostPlatforms = Executor->GetHostPlatforms();
+		TConstArrayView<FUtf8StringView> ExecutorHostPlatforms = Executor->GetHostPlatforms();
 		TArray<FBuildWorker*, TInlineAllocator<8>> FunctionWorkers;
 		Functions.MultiFind(MakeTuple(Function, FunctionVersion), FunctionWorkers);
 		for (FBuildWorker* Worker : FunctionWorkers)
diff --git a/Engine/Source/Developer/DerivedDataCache/Public/DerivedDataBuildWorker.h b/Engine/Source/Developer/DerivedDataCache/Public/DerivedDataBuildWorker.h
--- a/Engine/Source/Developer/DerivedDataCache/Public/DerivedDataBuildWorker.h
+++ b/Engine/Source/Developer/DerivedDataCache/Public/DerivedDataBuildWorker.h
@@ -3,62 +3,188 @@
 #pragma once
 
 #include "CoreTypes.h"
+#include "Containers/Array.h"
 #include "Containers/ArrayView.h"
 #include "Containers/StringFwd.h"
+#include "Containers/Utf8String.h"
 #include "DerivedDataBuildKey.h"
 #include "DerivedDataRequestTypes.h"
 #include "Features/IModularFeature.h"
+#include "IO/IoHash.h"
+#include "Misc/Guid.h"
+#include "Misc/NotNull.h"
+#include "Misc/OptionalFwd.h"
+#include "Serialization/CompactBinary.h"
+#include "Templates/Tuple.h"
 #include "UObject/NameTypes.h"
 
+#define UE_API DERIVEDDATACACHE_API
+
 class FCompressedBuffer;
-struct FGuid;
-struct FIoHash;
 
 template <typename FuncType> class TFunctionRef;
 template <typename FuncType> class TUniqueFunction;
 
 namespace UE::DerivedData { class FBuildAction; }
 namespace UE::DerivedData { class FBuildPolicy; }
+namespace UE::DerivedData { class FBuildWorkerBuilder; }
 namespace UE::DerivedData { class FOptionalBuildInputs; }
 namespace UE::DerivedData { class FOptionalBuildOutput; }
 namespace UE::DerivedData { class IBuild; }
+namespace UE::DerivedData { class IBuildWorkerResolver; }
 namespace UE::DerivedData { class IRequestOwner; }
 namespace UE::DerivedData { struct FBuildWorkerActionCompleteParams; }
-namespace UE::DerivedData { struct FBuildWorkerFileDataCompleteParams; }
+namespace UE::DerivedData { struct FBuildWorkerResolvedParams; }
 
 namespace UE::DerivedData
 {
 
 using FOnBuildWorkerActionComplete = TUniqueFunction<void (FBuildWorkerActionCompleteParams&& Params)>;
-using FOnBuildWorkerFileDataComplete = TUniqueFunction<void (FBuildWorkerFileDataCompleteParams&& Params)>;
+using FOnBuildWorkerResolved = TUniqueFunction<void (FBuildWorkerResolvedParams&& Params)>;
+
+/** A key that uniquely identifies a build worker. */
+struct FBuildWorkerKey final
+{
+	FIoHash Hash;
+};
 
-class FBuildWorker
+/**
+ * A build worker is an immutable reference to a build worker program and its dependencies.
+ *
+ * The purpose of a build worker is to capture which build functions a program contains and refer
+ * to everything required to execute the program, including files and environment.
+ *
+ * The key for the worker uniquely identifies the worker using a hash of the serialized compact
+ * binary representation of the worker.
+ *
+ * A build worker resolver is used to resolve file hashes to the referenced file data.
+ */
+class FBuildWorker final
 {
 public:
-	virtual ~FBuildWorker() = default;
-	virtual FStringView GetName() const = 0;
-	virtual FStringView GetPath() const = 0;
-	virtual FStringView GetHostPlatform() const = 0;
-	virtual FGuid GetBuildSystemVersion() const = 0;
-	virtual void FindFileData(TConstArrayView<FIoHash> RawHashes, IRequestOwner& Owner, FOnBuildWorkerFileDataComplete&& OnComplete) const = 0;
-	virtual void IterateFunctions(TFunctionRef<void (FUtf8StringView Name, const FGuid& Version)> Visitor) const = 0;
-	virtual void IterateFiles(TFunctionRef<void (FStringView Path, const FIoHash& RawHash, uint64 RawSize)> Visitor) const = 0;
-	virtual void IterateExecutables(TFunctionRef<void (FStringView Path, const FIoHash& RawHash, uint64 RawSize)> Visitor) const = 0;
-	virtual void IterateEnvironment(TFunctionRef<void (FStringView Name, FStringView Value)> Visitor) const = 0;
+	UE_API FBuildWorkerKey GetKey() const;
+	UE_API FUtf8StringView GetName() const;
+	UE_API FUtf8StringView GetPath() const;
+	UE_API FUtf8StringView GetHostPlatform() const;
+	UE_API FGuid GetBuildSystemVersion() const;
+	UE_API void IterateFunctions(TFunctionRef<void (FUtf8StringView Name, const FGuid& Version)> Visitor) const;
+	UE_API void IterateFiles(TFunctionRef<void (FUtf8StringView Path, const FIoHash& RawHash, uint64 RawSize)> Visitor) const;
+	UE_API void IterateExecutables(TFunctionRef<void (FUtf8StringView Path, const FIoHash& RawHash, uint64 RawSize)> Visitor) const;
+	UE_API void IterateEnvironment(TFunctionRef<void (FUtf8StringView Name, FUtf8StringView Value)> Visitor) const;
+
+	/** Resolve files and executables referenced by this worker, by raw hash. */
+	UE_API void Resolve(TConstArrayView<FIoHash> RawHashes, IRequestOwner& Owner, FOnBuildWorkerResolved&& OnResolved) const;
+
+	/** Saves the build worker to a compact binary object. */
+	inline const FCbObject& Save(FCbWriter& Writer) const
+	{
+		return Worker;
+	}
+
+	/**
+	 * Loads a build worker from a compact binary object.
+	 *
+	 * @param Key        The key of the worker being loaded. The key is validated during load.
+	 * @param Worker     An object saved from a build worker. Holds a reference and is cloned if not owned.
+	 * @param Resolver   A resolver used to find file data referenced by the build worker.
+	 * @return A valid build worker, or unset on error.
+	 */
+	UE_API static TOptional<FBuildWorker> Load(const FBuildWorkerKey& Key, FCbObject&& Worker, TUniquePtr<IBuildWorkerResolver> Resolver);
+
+	static constexpr bool bHasIntrusiveUnsetOptionalState = true;
+
+	inline explicit FBuildWorker(FIntrusiveUnsetOptionalState)
+	{
+	}
+
+	inline bool operator==(FIntrusiveUnsetOptionalState) const
+	{
+		return !Resolver;
+	}
+
+private:
+	explicit FBuildWorker(const FBuildWorkerKey& Key, FCbObject&& Worker, TUniquePtr<IBuildWorkerResolver> Resolver);
+
+	FCbObject Worker;
+	FBuildWorkerKey Key;
+	TUniquePtr<IBuildWorkerResolver> Resolver;
+
+	friend FBuildWorkerBuilder;
 };
 
+/** A builder to construct a build worker. */
 class FBuildWorkerBuilder
 {
 public:
+	// Destructor is virtual to allow an include-only dependency on this module.
 	virtual ~FBuildWorkerBuilder() = default;
-	virtual void SetName(FStringView Name) = 0;
-	virtual void SetPath(FStringView Path) = 0;
-	virtual void SetHostPlatform(FStringView Name) = 0;
-	virtual void SetBuildSystemVersion(const FGuid& Version) = 0;
-	virtual void AddFunction(FUtf8StringView Name, const FGuid& Version) = 0;
-	virtual void AddFile(FStringView Path, const FIoHash& RawHash, uint64 RawSize) = 0;
-	virtual void AddExecutable(FStringView Path, const FIoHash& RawHash, uint64 RawSize) = 0;
-	virtual void SetEnvironment(FStringView Name, FStringView Value) = 0;
+
+	inline void SetName(FUtf8StringView Name)
+	{
+		WorkerName = Name;
+	}
+
+	inline void SetPath(FUtf8StringView Path)
+	{
+		WorkerPath = Path;
+	}
+
+	inline void SetHostPlatform(FUtf8StringView Name)
+	{
+		HostPlatform = Name;
+	}
+
+	inline void SetBuildSystemVersion(const FGuid& Version)
+	{
+		BuildSystemVersion = Version;
+	}
+
+	inline void AddFunction(FUtf8StringView Name, const FGuid& Version)
+	{
+		Functions.Emplace(Name, Version);
+	}
+
+	inline void AddFile(FUtf8StringView Path, const FIoHash& RawHash, uint64 RawSize)
+	{
+		Files.Emplace(Path, RawHash, RawSize);
+	}
+
+	inline void AddSymbolFile(FUtf8StringView Path, const FIoHash& RawHash, uint64 RawSize)
+	{
+		SymbolFiles.Emplace(Path, RawHash, RawSize);
+	}
+
+	inline void AddExecutable(FUtf8StringView Path, const FIoHash& RawHash, uint64 RawSize)
+	{
+		Executables.Emplace(Path, RawHash, RawSize);
+	}
+
+	inline void SetEnvironment(FUtf8StringView Name, FUtf8StringView Value)
+	{
+		Environment.Emplace(Name, Value);
+	}
+
+	// Build is virtual to allow an include-only dependency on this module.
+	UE_API virtual FBuildWorker Build(TUniquePtr<IBuildWorkerResolver> Resolver);
+
+private:
+	FUtf8String WorkerName;
+	FUtf8String WorkerPath;
+	FUtf8String HostPlatform;
+	FGuid BuildSystemVersion;
+	TArray<TTuple<FUtf8String, FGuid>> Functions;
+	TArray<TTuple<FUtf8String, FIoHash, uint64>> Files;
+	TArray<TTuple<FUtf8String, FIoHash, uint64>> SymbolFiles;
+	TArray<TTuple<FUtf8String, FIoHash, uint64>> Executables;
+	TArray<TTuple<FUtf8String, FUtf8String>> Environment;
+};
+
+class IBuildWorkerResolver
+{
+public:
+	virtual ~IBuildWorkerResolver() = default;
+
+	virtual void Resolve(TConstArrayView<FIoHash> RawHashes, IRequestOwner& Owner, FOnBuildWorkerResolved&& OnResolved) = 0;
 };
 
 class IBuildWorkerFactory : public IModularFeature
@@ -68,8 +194,8 @@ public:
 
 	virtual ~IBuildWorkerFactory() = default;
 
-	virtual void Build(FBuildWorkerBuilder& Builder) = 0;
-	virtual void FindFileData(TConstArrayView<FIoHash> RawHashes, IRequestOwner& Owner, FOnBuildWorkerFileDataComplete&& OnComplete) = 0;
+	// Builder is provided to allow an include-only dependency on this module.
+	virtual TOptional<FBuildWorker> Build(FBuildWorkerBuilder& Builder) = 0;
 };
 
 class IBuildWorkerExecutor : public IModularFeature
@@ -88,10 +214,10 @@ public:
 		IRequestOwner& Owner,
 		FOnBuildWorkerActionComplete&& OnComplete) = 0;
 
-	virtual TConstArrayView<FStringView> GetHostPlatforms() const = 0;
+	virtual TConstArrayView<FUtf8StringView> GetHostPlatforms() const = 0;
 };
 
-struct FBuildWorkerActionCompleteParams
+struct FBuildWorkerActionCompleteParams final
 {
 	FBuildActionKey Key;
 	FOptionalBuildOutput&& Output;
@@ -99,10 +225,12 @@ struct FBuildWorkerActionCompleteParams
 	EStatus Status = EStatus::Error;
 };
 
-struct FBuildWorkerFileDataCompleteParams
+struct FBuildWorkerResolvedParams final
 {
 	TConstArrayView<FCompressedBuffer> Files;
 	EStatus Status = EStatus::Error;
 };
 
 } // UE::DerivedData
+
+#undef UE_API
diff --git a/Engine/Source/Developer/DesktopPlatform/Private/MultiPlatformTargetReceiptBuildWorkers.cpp b/Engine/Source/Developer/DesktopPlatform/Private/MultiPlatformTargetReceiptBuildWorkers.cpp
--- a/Engine/Source/Developer/DesktopPlatform/Private/MultiPlatformTargetReceiptBuildWorkers.cpp
+++ b/Engine/Source/Developer/DesktopPlatform/Private/MultiPlatformTargetReceiptBuildWorkers.cpp
@@ -2,32 +2,35 @@
 
 #include "MultiPlatformTargetReceiptBuildWorkers.h"
 
-FMultiPlatformTargetReceiptBuildWorkers::FMultiPlatformTargetReceiptBuildWorkers(const TCHAR* TargetReceiptFilePath)
-: bAllPlatformsInitialized(false)
+#include "Containers/StringView.h"
+#include "Containers/UnrealString.h"
+
+FMultiPlatformTargetReceiptBuildWorkers::FMultiPlatformTargetReceiptBuildWorkers(FStringView TargetReceiptFilePath)
 {
-	FString PathString = TargetReceiptFilePath;
-	if (!PathString.Contains(TEXT("$(Platform)")))
+	FString PathString(TargetReceiptFilePath);
+	if (PathString.Contains(TEXT("$(Platform)")))
 	{
-		new(&PlatformSpecificWorkerFactories[0]) FTargetReceiptBuildWorker(*PathString);
-		return;
+		bAllPlatformsInitialized = true;
+		PlatformSpecificWorkerFactories[SupportedPlatform_Win64].EmplaceUnchecked(PathString.Replace(TEXT("$(Platform)"), TEXT("Win64")));
+		PlatformSpecificWorkerFactories[SupportedPlatform_Mac].EmplaceUnchecked(PathString.Replace(TEXT("$(Platform)"), TEXT("Mac")));
+		PlatformSpecificWorkerFactories[SupportedPlatform_Linux].EmplaceUnchecked(PathString.Replace(TEXT("$(Platform)"), TEXT("Linux")));
+	}
+	else
+	{
+		PlatformSpecificWorkerFactories[0].EmplaceUnchecked(PathString);
 	}
-
-	new(&PlatformSpecificWorkerFactories[SupportedPlatform_Win64]) FTargetReceiptBuildWorker(*PathString.Replace(TEXT("$(Platform)"), TEXT("Win64")));
-	new(&PlatformSpecificWorkerFactories[SupportedPlatform_Mac]) FTargetReceiptBuildWorker(*PathString.Replace(TEXT("$(Platform)"), TEXT("Mac")));
-	new(&PlatformSpecificWorkerFactories[SupportedPlatform_Linux]) FTargetReceiptBuildWorker(*PathString.Replace(TEXT("$(Platform)"), TEXT("Linux")));
-	bAllPlatformsInitialized = true;
 }
 
 FMultiPlatformTargetReceiptBuildWorkers::~FMultiPlatformTargetReceiptBuildWorkers()
 {
 	if (bAllPlatformsInitialized)
 	{
-		((FTargetReceiptBuildWorker*)&PlatformSpecificWorkerFactories[SupportedPlatform_Win64])->~FTargetReceiptBuildWorker();
-		((FTargetReceiptBuildWorker*)&PlatformSpecificWorkerFactories[SupportedPlatform_Mac])->~FTargetReceiptBuildWorker();
-		((FTargetReceiptBuildWorker*)&PlatformSpecificWorkerFactories[SupportedPlatform_Linux])->~FTargetReceiptBuildWorker();
+		PlatformSpecificWorkerFactories[SupportedPlatform_Win64].DestroyUnchecked();
+		PlatformSpecificWorkerFactories[SupportedPlatform_Mac].DestroyUnchecked();
+		PlatformSpecificWorkerFactories[SupportedPlatform_Linux].DestroyUnchecked();
 	}
 	else
 	{
-		((FTargetReceiptBuildWorker*)&PlatformSpecificWorkerFactories[0])->~FTargetReceiptBuildWorker();
+		PlatformSpecificWorkerFactories[0].DestroyUnchecked();
 	}
 }
diff --git a/Engine/Source/Developer/DesktopPlatform/Private/TargetReceiptBuildWorker.cpp b/Engine/Source/Developer/DesktopPlatform/Private/TargetReceiptBuildWorker.cpp
--- a/Engine/Source/Developer/DesktopPlatform/Private/TargetReceiptBuildWorker.cpp
+++ b/Engine/Source/Developer/DesktopPlatform/Private/TargetReceiptBuildWorker.cpp
@@ -4,29 +4,74 @@
 
 #include "Algo/Sort.h"
 #include "Algo/Unique.h"
+#include "Async/TaskGraphInterfaces.h"
 #include "Compression/CompressedBuffer.h"
 #include "Compression/OodleDataCompression.h"
+#include "Containers/Array.h"
+#include "Containers/Map.h"
+#include "Containers/StringView.h"
+#include "Containers/UnrealString.h"
+#include "Containers/Utf8String.h"
 #include "DerivedDataBuildWorker.h"
 #include "Features/IModularFeatures.h"
 #include "HAL/FileManager.h"
 #include "HAL/Platform.h"
 #include "HAL/PlatformProcess.h"
 #include "HAL/PlatformTime.h"
+#include "IO/IoHash.h"
+#include "Logging/StructuredLog.h"
+#include "Misc/Guid.h"
 #include "Misc/Paths.h"
 #include "Misc/PathViews.h"
 #include "Misc/StringBuilder.h"
 #include "Serialization/CompactBinary.h"
 #include "Serialization/CompactBinarySerialization.h"
 #include "TargetReceipt.h"
+#include "Templates/Tuple.h"
+
+#include <atomic>
 
 DEFINE_LOG_CATEGORY_STATIC(LogTargetReceiptBuildWorker, Log, All);
 
-// This needs to match the version exported in DerivedDataBuildWorker.Build.cs 
-const FGuid FTargetReceiptBuildWorker::WorkerReceiptVersion(TEXT("dab5352e-a5a7-4793-a7a3-1d4acad6aff2"));
+class FTargetReceiptBuildWorkerLoader
+{
+	// This needs to match the version exported in DerivedDataBuildWorker.Build.cs.
+	inline static const FGuid WorkerReceiptVersion{TEXT("dab5352e-a5a7-4793-a7a3-1d4acad6aff2")};
+
+public:
+	bool Load(FStringView TargetReceiptFilePath);
+	void AbortLoad();
+
+	UE::DerivedData::FBuildWorker Build(UE::DerivedData::FBuildWorkerBuilder& Builder) const;
+
+private:
+	bool TryAddExecutablePath(FStringView WidePath);
 
-FTargetReceiptBuildWorker::FWorkerPath::FWorkerPath(FStringView InLocalPathRoot, FStringView InLocalPathSuffix, FStringView InRemotePathRoot, FStringView InRemotePathSuffix)
+	/**
+	 * Encapsulates a path used by the worker that has both a local and remote representation.
+	 */
+	struct FWorkerPath
+	{
+		explicit FWorkerPath(FUtf8StringView InLocalPathRoot, FUtf8StringView InLocalPathSuffix, FUtf8StringView InRemotePathRoot = {}, FUtf8StringView InRemotePathSuffix = {});
+
+		FUtf8String LocalPath;
+		FUtf8String RemotePath;
+	};
+
+	FUtf8String Name;
+	FUtf8String Platform;
+	FGuid BuildSystemVersion;
+	TArray<TTuple<FUtf8String, FGuid>> FunctionVersions;
+	using FWorkerPathMeta = TTuple<FIoHash, uint64>;
+	TArray<FWorkerPathMeta> ExecutableMeta;
+	TArray<FWorkerPath> ExecutablePaths;
+	TMap<FUtf8String, FUtf8String> EnvironmentVariables;
+	std::atomic<bool> bAbortRequested = false;
+};
+
+FTargetReceiptBuildWorkerLoader::FWorkerPath::FWorkerPath(FUtf8StringView InLocalPathRoot, FUtf8StringView InLocalPathSuffix, FUtf8StringView InRemotePathRoot, FUtf8StringView InRemotePathSuffix)
 {
-	TStringBuilder<128> PathBuilder;
+	TUtf8StringBuilder<128> PathBuilder;
 	FPathViews::Append(PathBuilder, InLocalPathRoot, InLocalPathSuffix);
 	LocalPath = PathBuilder;
 
@@ -37,181 +82,125 @@ FTargetReceiptBuildWorker::FWorkerPath::FWorkerPath(FStringView InLocalPathRoot,
 	RemotePath = PathBuilder;
 }
 
-class FPopulateWorkerFromTargetReceiptTask
+#if WITH_EDITOR
+class FTargetReceiptBuildWorkerFactory final : public UE::DerivedData::IBuildWorkerFactory
 {
 public:
-	FPopulateWorkerFromTargetReceiptTask(const TCHAR* InTargetReceiptFilePath, FTargetReceiptBuildWorker* InWorker)
-	: TargetReceiptFilePath(InTargetReceiptFilePath)
-	, Worker(InWorker)
+	explicit FTargetReceiptBuildWorkerFactory(FStringView TargetReceiptFilePath);
+
+	~FTargetReceiptBuildWorkerFactory()
 	{
+		Loader.AbortLoad();
+		if (LoadTargetReceiptTask.IsValid())
+		{
+			LoadTargetReceiptTask->Wait();
+		}
+
+		if (bRegistered)
+		{
+			IModularFeatures::Get().UnregisterModularFeature(UE::DerivedData::IBuildWorkerFactory::FeatureName, this);
+		}
 	}
 
-	FORCEINLINE TStatId GetStatId() const { RETURN_QUICK_DECLARE_CYCLE_STAT(FPopulateWorkerFromTargetReceiptTask, STATGROUP_TaskGraphTasks); }
-	ENamedThreads::Type GetDesiredThread() { return ENamedThreads::Type::AnyThread; }
-	static ESubsequentsMode::Type GetSubsequentsMode() { return ESubsequentsMode::FireAndForget; }
+	void LoadAndRegister(FStringView TargetReceiptFilePath)
+	{
+		if (Loader.Load(TargetReceiptFilePath))
+		{
+			check(!bRegistered);
+			bRegistered = true;
+			IModularFeatures::Get().RegisterModularFeature(UE::DerivedData::IBuildWorkerFactory::FeatureName, this);
+		}
+	}
 
-	void DoTask(ENamedThreads::Type CurrentThread, const FGraphEventRef& MyCompletionGraphEvent)
+	TOptional<UE::DerivedData::FBuildWorker> Build(UE::DerivedData::FBuildWorkerBuilder& Builder) final
 	{
-		Worker->PopulateWorkerFromReceipt(TargetReceiptFilePath);
+		return Loader.Build(Builder);
 	}
 
 private:
-	const FString TargetReceiptFilePath;
-	FTargetReceiptBuildWorker* Worker;
+	FGraphEventRef LoadTargetReceiptTask;
+	FTargetReceiptBuildWorkerLoader Loader;
+	bool bRegistered = false;
 };
 
-class FTargetReceiptBuildWorkerFactory final : public UE::DerivedData::IBuildWorkerFactory
+class FPopulateWorkerFromTargetReceiptTask
 {
 public:
-	const FTargetReceiptBuildWorker* GetWorker()
+	explicit FPopulateWorkerFromTargetReceiptTask(FStringView InTargetReceiptFilePath, FTargetReceiptBuildWorkerFactory* InFactory)
+		: TargetReceiptFilePath(InTargetReceiptFilePath)
+		, Factory(InFactory)
 	{
-		return (FTargetReceiptBuildWorker*)((uintptr_t)this - offsetof(class FTargetReceiptBuildWorker, InternalFactory));
 	}
 
-	void Build(UE::DerivedData::FBuildWorkerBuilder& Builder) final
-	{
-		const FTargetReceiptBuildWorker* Worker = GetWorker();
-		Builder.SetName(Worker->Name);
-		Builder.SetPath(Worker->ExecutablePaths[0].GetRemotePath());
-		Builder.SetHostPlatform(Worker->Platform);
-
-		for (int32 ExecutableIndex = 0; ExecutableIndex < Worker->ExecutablePaths.Num(); ++ExecutableIndex)
-		{
-			Builder.AddExecutable(Worker->ExecutablePaths[ExecutableIndex].GetRemotePath(), Worker->ExecutableMeta[ExecutableIndex].Key, Worker->ExecutableMeta[ExecutableIndex].Value);
-		}
-
-		for (const TPair<FString, FString>& EnvironmentVariable : Worker->EnvironmentVariables)
-		{
-			Builder.SetEnvironment(EnvironmentVariable.Key, EnvironmentVariable.Value);
-		}
-
-		Builder.SetBuildSystemVersion(Worker->BuildSystemVersion);
+	inline TStatId GetStatId() const { RETURN_QUICK_DECLARE_CYCLE_STAT(FPopulateWorkerFromTargetReceiptTask, STATGROUP_TaskGraphTasks); }
+	constexpr ENamedThreads::Type GetDesiredThread() { return ENamedThreads::Type::AnyThread; }
+	constexpr static ESubsequentsMode::Type GetSubsequentsMode() { return ESubsequentsMode::FireAndForget; }
 
-		for (const TPair<FString, FGuid>& FunctionVersion : Worker->FunctionVersions)
-		{
-			Builder.AddFunction(FTCHARToUTF8(FunctionVersion.Key), FunctionVersion.Value);
-		}
-	}
-
-	void FindFileData(TConstArrayView<FIoHash> RawHashes, UE::DerivedData::IRequestOwner& Owner, UE::DerivedData::FOnBuildWorkerFileDataComplete&& OnComplete) final
+	void DoTask(ENamedThreads::Type CurrentThread, const FGraphEventRef& MyCompletionGraphEvent)
 	{
-		if (OnComplete)
-		{
-			const FTargetReceiptBuildWorker* Worker = GetWorker();
-			TArray<FCompressedBuffer> FileDataBuffers;
-			UE::DerivedData::FBuildWorkerFileDataCompleteParams CompleteParams;
-			for (const FIoHash& RawHash : RawHashes)
-			{
-				for (int32 ExecutableIndex = 0; ExecutableIndex < Worker->ExecutablePaths.Num(); ++ExecutableIndex)
-				{
-					const FTargetReceiptBuildWorker::FWorkerPathMeta& Meta = Worker->ExecutableMeta[ExecutableIndex];
-					if (Meta.Key == RawHash)
-					{
-						if (TUniquePtr<FArchive> Ar{IFileManager::Get().CreateFileReader(*Worker->ExecutablePaths[ExecutableIndex].GetLocalPath(), FILEREAD_Silent)})
-						{
-							const int64 TotalSize = Ar->TotalSize();
-							FUniqueBuffer MutableBuffer = FUniqueBuffer::Alloc(uint64(TotalSize));
-							Ar->Serialize(MutableBuffer.GetData(), TotalSize);
-							if (Ar->Close())
-							{
-								FileDataBuffers.Emplace(FCompressedBuffer::Compress(
-										MutableBuffer.MoveToShared(),
-										ECompressedBufferCompressor::NotSet,
-										ECompressedBufferCompressionLevel::None
-									));
-							}
-						}
-					}
-				}
-			}
-
-			if (FileDataBuffers.Num() == RawHashes.Num())
-			{
-				CompleteParams.Status = UE::DerivedData::EStatus::Ok;
-				CompleteParams.Files = FileDataBuffers;
-				OnComplete(MoveTemp(CompleteParams));
-			}
-			else
-			{
-				CompleteParams.Status = UE::DerivedData::EStatus::Error;
-				OnComplete(MoveTemp(CompleteParams));
-			}
-		}
+		Factory->LoadAndRegister(TargetReceiptFilePath);
 	}
+
+private:
+	const FString TargetReceiptFilePath;
+	FTargetReceiptBuildWorkerFactory* Factory;
 };
 
-FTargetReceiptBuildWorker::FTargetReceiptBuildWorker(const TCHAR* TargetReceiptFilePath)
-: bAbortRequested(false)
-, bEnabled(false)
+FTargetReceiptBuildWorkerFactory::FTargetReceiptBuildWorkerFactory(FStringView TargetReceiptFilePath)
 {
-#if WITH_EDITOR
-	static_assert(sizeof(FTargetReceiptBuildWorkerFactory) == sizeof(IPureVirtual), "FTargetReceiptBuildWorkerFactory type must always compile to something equivalent to an IPureVirtual size.");
-	FTargetReceiptBuildWorkerFactory* WorkerFactory = GetWorkerFactory();
-	new(WorkerFactory) FTargetReceiptBuildWorkerFactory();
-	PopulateTaskRef = TGraphTask<FPopulateWorkerFromTargetReceiptTask>::CreateTask().ConstructAndDispatchWhenReady(TargetReceiptFilePath, this);
-#endif
+	LoadTargetReceiptTask = TGraphTask<FPopulateWorkerFromTargetReceiptTask>::CreateTask().ConstructAndDispatchWhenReady(TargetReceiptFilePath, this);
 }
+#endif // WITH_EDITOR
 
-FTargetReceiptBuildWorker::~FTargetReceiptBuildWorker()
+FTargetReceiptBuildWorker::FTargetReceiptBuildWorker(FStringView TargetReceiptFilePath)
 {
 #if WITH_EDITOR
-	bAbortRequested = true;
-	if (PopulateTaskRef.IsValid())
-	{
-		PopulateTaskRef->Wait();
-	}
-
-	FTargetReceiptBuildWorkerFactory* WorkerFactory = GetWorkerFactory();
-	if (bEnabled)
-	{
-		IModularFeatures::Get().UnregisterModularFeature(UE::DerivedData::IBuildWorkerFactory::FeatureName, WorkerFactory);
-	}
-	WorkerFactory->~FTargetReceiptBuildWorkerFactory();
+	Factory = MakePimpl<FTargetReceiptBuildWorkerFactory>(TargetReceiptFilePath);
 #endif
 }
 
-FTargetReceiptBuildWorkerFactory* FTargetReceiptBuildWorker::GetWorkerFactory()
-{
-	return (FTargetReceiptBuildWorkerFactory*)&InternalFactory;
-}
+FTargetReceiptBuildWorker::~FTargetReceiptBuildWorker() = default;
 
-bool FTargetReceiptBuildWorker::TryAddExecutablePath(FStringView Path)
+bool FTargetReceiptBuildWorkerLoader::TryAddExecutablePath(FStringView WidePath)
 {
-	if (Path.IsEmpty())
+	TUtf8StringBuilder<256> Path(InPlace, WidePath);
+
+	if (Path.ToView().IsEmpty())
 	{
-		UE_LOG(LogTargetReceiptBuildWorker, Warning, TEXT("Empty executable path encountered while populating worker '%s': '%.*s'"), *Name, Path.Len(), Path.GetData());
+		UE_LOGF(LogTargetReceiptBuildWorker, Warning, "Empty executable path encountered while populating worker '%s': '%s'", *Name, *Path);
 		return false;
 	}
 
-	if (!Path.StartsWith(TEXT("$(")))
+	if (!Path.ToView().StartsWith(ANSITEXTVIEW("$(")))
 	{
-		UE_LOG(LogTargetReceiptBuildWorker, Warning, TEXT("Absolute or unprefixed executable path variable encountered while populating worker '%s': '%.*s'"), *Name, Path.Len(), Path.GetData());
+		UE_LOGF(LogTargetReceiptBuildWorker, Warning, "Absolute or unprefixed executable path variable encountered while populating worker '%s': '%s'", *Name, *Path);
 		return false;
 	}
 
-	int32 EndVariableIndex = Path.Find(TEXT(")"));
+	int32 EndVariableIndex = Path.ToView().Find(ANSITEXTVIEW(")"));
 	if (EndVariableIndex == INDEX_NONE)
 	{
-		UE_LOG(LogTargetReceiptBuildWorker, Warning, TEXT("Malformed executable path variable encountered while populating worker '%s': '%.*s'"), *Name, Path.Len(), Path.GetData());
+		UE_LOGF(LogTargetReceiptBuildWorker, Warning, "Malformed executable path variable encountered while populating worker '%s': '%s'", *Name, *Path);
 		return false;
 	}
 
-	FStringView PrefixVariable = Path.Mid(2, EndVariableIndex - 2);
-	if (PrefixVariable == TEXT("EngineDir"))
+	FUtf8StringView PrefixVariable = Path.ToView().Mid(2, EndVariableIndex - 2);
+	if (PrefixVariable == ANSITEXTVIEW("EngineDir"))
 	{
-		ExecutablePaths.Emplace(FPaths::EngineDir(), Path.RightChop(EndVariableIndex + 2), TEXT("Engine"));
+		ExecutablePaths.Emplace(FUtf8String(FPaths::EngineDir()), Path.ToView().RightChop(EndVariableIndex + 2), ANSITEXTVIEW("Engine"));
 	}
-	else if (PrefixVariable == TEXT("ProjectDir"))
+	else if (PrefixVariable == ANSITEXTVIEW("ProjectDir"))
 	{
-		FString ProjectDir = FPaths::ConvertRelativePathToFull(FPaths::ProjectDir());
+		FString ProjectDir = FPaths::ProjectDir();
+		FString RelativeProjectDir = FPaths::ConvertRelativePathToFull(ProjectDir);
 		FString RootDir = FPaths::ConvertRelativePathToFull(FPaths::RootDir());
-		if (FPaths::MakePathRelativeTo(ProjectDir, *RootDir))
+		if (FPaths::MakePathRelativeTo(RelativeProjectDir, *RootDir))
 		{
-			ExecutablePaths.Emplace(FPaths::ProjectDir(), Path.RightChop(EndVariableIndex + 2), ProjectDir);
+			ExecutablePaths.Emplace(FUtf8String(ProjectDir), Path.ToView().RightChop(EndVariableIndex + 2), WriteToUtf8String<64>(RelativeProjectDir));
 		}
 		else
 		{
-			ExecutablePaths.Emplace(FPaths::ProjectDir(), Path.RightChop(EndVariableIndex + 2), TEXT("Project"));
+			ExecutablePaths.Emplace(FUtf8String(ProjectDir), Path.ToView().RightChop(EndVariableIndex + 2), ANSITEXTVIEW("Project"));
 		}
 	}
 	else
@@ -220,43 +209,42 @@ bool FTargetReceiptBuildWorker::TryAddExecutablePath(FStringView Path)
 		FString EnvironmentVariableValue = FPlatformMisc::GetEnvironmentVariable(*TempVariableName);
 		if (EnvironmentVariableValue.IsEmpty())
 		{
-			UE_LOG(LogTargetReceiptBuildWorker, Warning, TEXT("Unknown executable path variable encountered while populating worker '%s': '%s'"), *Name, *TempVariableName);
+			UE_LOGF(LogTargetReceiptBuildWorker, Warning, "Unknown executable path variable encountered while populating worker '%s': '%ls'", *Name, *TempVariableName);
 			return false;
 		}
-		ExecutablePaths.Emplace(EnvironmentVariableValue, Path.RightChop(EndVariableIndex + 2), TempVariableName);
+		ExecutablePaths.Emplace(FUtf8String(EnvironmentVariableValue), Path.ToView().RightChop(EndVariableIndex + 2), PrefixVariable);
 		// By having a dependency that starts with an environment variable, we implicitly create an environment variable to be remoted.
 		// The variable will have the same name as the variable we used, and its value will be a path relative to the platform specific binaries directory with the same name as the variable.
-		EnvironmentVariables.FindOrAdd(TempVariableName, FString(TEXT("..\\..\\..\\")) + TempVariableName);
+		EnvironmentVariables.FindOrAdd(FUtf8String(PrefixVariable), FUtf8String(ANSITEXTVIEW("..\\..\\..\\")) + PrefixVariable);
 	}
 
 	return true;
 }
 
-void FTargetReceiptBuildWorker::PopulateWorkerFromReceipt(const FString& TargetReceiptFilePath)
+bool FTargetReceiptBuildWorkerLoader::Load(FStringView TargetReceiptFilePath)
 {
-	if (bAbortRequested)
-	{
-		return;
-	}
-
 	IFileManager& FileManager = IFileManager::Get();
 	FString FinalTargetReceiptFilePath;
 	FString FinalPublishedVersionFilePath;
 	FString FinalGeneratedVersionFilePath;
-	if (TargetReceiptFilePath.StartsWith(TEXT("$(ProjectDir)")))
+	const FStringView EngineDirPrefix = TEXTVIEW("$(EngineDir)");
+	const FStringView ProjectDirPrefix = TEXTVIEW("$(ProjectDir)");
+	if (TargetReceiptFilePath.StartsWith(ProjectDirPrefix))
 	{
-		FinalTargetReceiptFilePath = FPaths::ConvertRelativePathToFull(TargetReceiptFilePath.Replace(TEXT("$(ProjectDir)"), *FPaths::ProjectDir()));
+		const FStringView TargetReceiptSuffix = TargetReceiptFilePath.RightChop(ProjectDirPrefix.Len());
+		FinalTargetReceiptFilePath = FPaths::ConvertRelativePathToFull(FPaths::ProjectDir() + TargetReceiptSuffix);
 		FinalPublishedVersionFilePath = FinalTargetReceiptFilePath + TEXT(".BuildVersion");
-		FinalGeneratedVersionFilePath = FPaths::ConvertRelativePathToFull(TargetReceiptFilePath.Replace(TEXT("$(ProjectDir)"), *FPaths::Combine(FPaths::ProjectSavedDir(), TEXT("GeneratedWorkerBuildVersions"))) + TEXT(".BuildVersion"));
+		FinalGeneratedVersionFilePath = FPaths::ConvertRelativePathToFull(FPaths::Combine(FPaths::ProjectSavedDir(), TEXT("GeneratedWorkerBuildVersions")) + TargetReceiptSuffix + TEXT(".BuildVersion"));
 	}
-	else if (TargetReceiptFilePath.StartsWith(TEXT("$(EngineDir)")))
+	else if (TargetReceiptFilePath.StartsWith(EngineDirPrefix))
 	{
 		// Target receipts declared to be in the engine directory may end up in the project directory if they were compiled
 		// with a uproject file on the commandline.  For this reason we allow for the possibility that the target receipt path
 		// can be tried with both the true engine directory or the engine directory treated as the project directory.
 		// Pick whichever is newer (more recently modified).
-		FString ReceiptEnginePath = TargetReceiptFilePath.Replace(TEXT("$(EngineDir)"), *FPaths::EngineDir());
-		FString ReceiptProjectPath = TargetReceiptFilePath.Replace(TEXT("$(EngineDir)"), *FPaths::ProjectDir());
+		const FStringView TargetReceiptSuffix = TargetReceiptFilePath.RightChop(EngineDirPrefix.Len());
+		FString ReceiptEnginePath = FPaths::EngineDir() + TargetReceiptSuffix;
+		FString ReceiptProjectPath = FPaths::ProjectDir() + TargetReceiptSuffix;
 		FDateTime ReceiptEngineTimeStamp;
 		FDateTime ReceiptProjectTimeStamp;
 		FileManager.GetTimeStampPair(*ReceiptEnginePath, *ReceiptProjectPath, ReceiptEngineTimeStamp, ReceiptProjectTimeStamp);
@@ -264,48 +252,48 @@ void FTargetReceiptBuildWorker::PopulateWorkerFromReceipt(const FString& TargetR
 		{
 			FinalTargetReceiptFilePath = FPaths::ConvertRelativePathToFull(ReceiptProjectPath);
 			FinalPublishedVersionFilePath = FinalTargetReceiptFilePath + TEXT(".BuildVersion");
-			FinalGeneratedVersionFilePath = FPaths::ConvertRelativePathToFull(TargetReceiptFilePath.Replace(TEXT("$(EngineDir)"), *FPaths::Combine(FPaths::ProjectSavedDir(), TEXT("GeneratedWorkerBuildVersions"))) + TEXT(".BuildVersion"));
+			FinalGeneratedVersionFilePath = FPaths::ConvertRelativePathToFull(FPaths::Combine(FPaths::ProjectSavedDir(), TEXT("GeneratedWorkerBuildVersions")) + TargetReceiptSuffix + TEXT(".BuildVersion"));
 		}
 		else
 		{
 			FinalTargetReceiptFilePath = FPaths::ConvertRelativePathToFull(ReceiptEnginePath);
 			FinalPublishedVersionFilePath = FinalTargetReceiptFilePath + TEXT(".BuildVersion");
-			FinalGeneratedVersionFilePath = FPaths::ConvertRelativePathToFull(TargetReceiptFilePath.Replace(TEXT("$(EngineDir)"), *FPaths::Combine(FPaths::EngineSavedDir(), TEXT("GeneratedWorkerBuildVersions"))) + TEXT(".BuildVersion"));
+			FinalGeneratedVersionFilePath = FPaths::ConvertRelativePathToFull(FPaths::Combine(FPaths::EngineSavedDir(), TEXT("GeneratedWorkerBuildVersions")) + TargetReceiptSuffix + TEXT(".BuildVersion"));
 		}
 	}
 	else
 	{
-		UE_LOG(LogTargetReceiptBuildWorker, Warning, TEXT("Unknown prefix for target receipt path: '%s'"), *TargetReceiptFilePath);
-		return;
+		UE_LOGFMT(LogTargetReceiptBuildWorker, Warning, "Unknown prefix for target receipt path: '{TargetReceipt}'", TargetReceiptFilePath);
+		return false;
 	}
 	
 	if (!FPaths::FileExists(FinalTargetReceiptFilePath))
 	{
 		// This case should return without emitting a warning.
-		return;
+		return false;
 	}
 
 	FTargetReceipt TargetReceipt;
 	if (!TargetReceipt.Read(FinalTargetReceiptFilePath, false))
 	{
-		UE_LOG(LogTargetReceiptBuildWorker, Warning, TEXT("Target receipt failed to parse: '%s'"), *TargetReceiptFilePath);
-		return;
+		UE_LOGFMT(LogTargetReceiptBuildWorker, Warning, "Target receipt failed to parse: '{TargetReceipt}'", TargetReceiptFilePath);
+		return false;
 	}
 
-	if (bAbortRequested)
+	if (bAbortRequested.load(std::memory_order_relaxed))
 	{
-		return;
+		return false;
 	}
 
-	Name = TargetReceipt.TargetName;
-	Platform = TargetReceipt.Platform;
+	Name = FUtf8String(TargetReceipt.TargetName);
+	Platform = FUtf8String(TargetReceipt.Platform);
 	if (!TryAddExecutablePath(TargetReceipt.Launch))
 	{
-		return;
+		return false;
 	}
 
 	FGuid ReceiptInlineBuildSystemVersion;
-	TArray<TPair<FString, FGuid>> ReceiptInlineFunctionVersions;
+	TArray<TPair<FUtf8String, FGuid>> ReceiptInlineFunctionVersions;
 	bool bUsableBuildWorker = false;
 	for (const FReceiptProperty& Property : TargetReceipt.AdditionalProperties)
 	{
@@ -318,7 +306,7 @@ void FTargetReceiptBuildWorker::PopulateWorkerFromReceipt(const FString& TargetR
 			if (!FGuid::Parse(Property.Value, ReceiptInlineBuildSystemVersion))
 			{
 				UE_LOG(LogTargetReceiptBuildWorker, Warning, TEXT("Target receipt contains inline build system version guid that could not be parsed: '%s'"), *Property.Value);
-				return;
+				return false;
 			}
 		}
 		else if (Property.Name == TEXT("DerivedDataBuildWorkerFunctionVersion"))
@@ -327,14 +315,14 @@ void FTargetReceiptBuildWorker::PopulateWorkerFromReceipt(const FString& TargetR
 			if (2 != Property.Value.ParseIntoArray(FunctionComponents, TEXT(":")))
 			{
 				UE_LOG(LogTargetReceiptBuildWorker, Warning, TEXT("Target receipt contains inline function version declaration that could not be parsed: '%s'"), *Property.Value);
-				return;
+				return false;
 			}
 
 			FGuid VersionGuid;
 			if (!FGuid::Parse(FunctionComponents[1], VersionGuid))
 			{
 				UE_LOG(LogTargetReceiptBuildWorker, Warning, TEXT("Target receipt contains inline function version guid that could not be parsed: '%s'"), *Property.Value);
-				return;
+				return false;
 			}
 
 			ReceiptInlineFunctionVersions.Emplace(FunctionComponents[0], MoveTemp(VersionGuid));
@@ -344,14 +332,14 @@ void FTargetReceiptBuildWorker::PopulateWorkerFromReceipt(const FString& TargetR
 			if (!TryAddExecutablePath(Property.Value))
 			{
 				UE_LOG(LogTargetReceiptBuildWorker, Warning, TEXT("Target receipt referenced an unstaged dependency that was not usable: '%s'"), *Property.Value);
-				return;
+				return false;
 			}
 		}
 	}
 
 	if (!bUsableBuildWorker)
 	{
-		return;
+		return false;
 	}
 
 	for (const FBuildProduct& BuildProduct : TargetReceipt.BuildProducts)
@@ -364,7 +352,7 @@ void FTargetReceiptBuildWorker::PopulateWorkerFromReceipt(const FString& TargetR
 		if (!TryAddExecutablePath(BuildProduct.Path))
 		{
 			UE_LOG(LogTargetReceiptBuildWorker, Warning, TEXT("Target receipt referenced a build product that was not usable: '%s'"), *BuildProduct.Path);
-			return;
+			return false;
 		}
 	}
 
@@ -378,19 +366,19 @@ void FTargetReceiptBuildWorker::PopulateWorkerFromReceipt(const FString& TargetR
 		if (!TryAddExecutablePath(RuntimeDependency.Path))
 		{
 			UE_LOG(LogTargetReceiptBuildWorker, Warning, TEXT("Target receipt referenced a runtime dependency that was not usable: '%s'"), *RuntimeDependency.Path);
-			return;
+			return false;
 		}
 	}
 	
-	if (bAbortRequested)
+	if (bAbortRequested.load(std::memory_order_relaxed))
 	{
-		return;
+		return false;
 	}
 
 	// Build a sorted list of unique paths because the target receipt may reference the same file in multiple places.
 	// Skip the first path because it is assumed elsewhere to be the launch path.
-	Algo::SortBy(MakeArrayView(ExecutablePaths).RightChop(1), &FWorkerPath::GetLocalPath);
-	const int32 KeepPaths = Algo::UniqueBy(ExecutablePaths, &FWorkerPath::GetLocalPath);
+	Algo::SortBy(MakeArrayView(ExecutablePaths).RightChop(1), &FWorkerPath::LocalPath);
+	const int32 KeepPaths = Algo::UniqueBy(ExecutablePaths, &FWorkerPath::LocalPath);
 	ExecutablePaths.RemoveAt(KeepPaths, ExecutablePaths.Num() - KeepPaths, EAllowShrinking::No);
 
 	// Open executable files to check for existence and gather hash info for them
@@ -402,7 +390,7 @@ void FTargetReceiptBuildWorker::PopulateWorkerFromReceipt(const FString& TargetR
 	{
 		FWorkerPathMeta& Meta = ExecutableMeta[MetaIndex];
 
-		FString LocalExecutablePath = ExecutablePath.GetLocalPath();
+		TStringBuilder<256> LocalExecutablePath(InPlace, ExecutablePath.LocalPath);
 
 		if (TUniquePtr<FArchive> Ar{FileManager.CreateFileReader(*LocalExecutablePath, FILEREAD_Silent)})
 		{
@@ -421,7 +409,7 @@ void FTargetReceiptBuildWorker::PopulateWorkerFromReceipt(const FString& TargetR
 		else
 		{
 			UE_LOG(LogTargetReceiptBuildWorker, Warning, TEXT("Executable file not present: '%s'"), *LocalExecutablePath);
-			return;
+			return false;
 		}
 		++MetaIndex;
 	}
@@ -448,7 +436,7 @@ void FTargetReceiptBuildWorker::PopulateWorkerFromReceipt(const FString& TargetR
 			else
 			{
 				UE_LOG(LogTargetReceiptBuildWorker, Warning, TEXT("Could not read published version file: '%s'"), *TargetReceipt.Launch);
-				return;
+				return false;
 			}
 		}
 		else
@@ -458,30 +446,30 @@ void FTargetReceiptBuildWorker::PopulateWorkerFromReceipt(const FString& TargetR
 			if (TargetReceipt.Platform != FPlatformProcess::GetBinariesSubdirectory())
 			{
 				UE_LOG(LogTargetReceiptBuildWorker, Log, TEXT("Target lacked published version information and isn't made for the current platform: '%s'"), *TargetReceipt.Launch);
-				return;
+				return false;
 			}
 
 			FDateTime VersionFileModificationTime = FileManager.GetTimeStamp(*FinalGeneratedVersionFilePath);
 			if (VersionFileModificationTime < NewestExecutableModificationTime)
 			{
 				FString CommandLine = TEXT("-Version=\"") + FinalGeneratedVersionFilePath + TEXT("\"");
-				FProcHandle WorkerProcHandle = FPlatformProcess::CreateProc(*FPaths::ConvertRelativePathToFull(ExecutablePaths[0].GetLocalPath()), *CommandLine, true, true, true, nullptr, 0, nullptr, nullptr);
+				FProcHandle WorkerProcHandle = FPlatformProcess::CreateProc(*FPaths::ConvertRelativePathToFull(FString(ExecutablePaths[0].LocalPath)), *CommandLine, true, true, true, nullptr, 0, nullptr, nullptr);
 				if (!WorkerProcHandle.IsValid())
 				{
 					UE_LOG(LogTargetReceiptBuildWorker, Warning, TEXT("Failed to launch worker: '%s'"), *TargetReceipt.Launch);
-					return;
+					return false;
 				}
 
 				uint64 WorkerStartTime = FPlatformTime::Cycles64();
 				while (FPlatformProcess::IsProcRunning(WorkerProcHandle))
 				{
 					FPlatformProcess::Sleep(0.1f);
-					if (bAbortRequested || (FPlatformTime::ToSeconds64(FPlatformTime::Cycles64()-WorkerStartTime) > 10))
+					if (bAbortRequested.load(std::memory_order_relaxed) || (FPlatformTime::ToSeconds64(FPlatformTime::Cycles64()-WorkerStartTime) > 10))
 					{
 						UE_LOG(LogTargetReceiptBuildWorker, Warning, TEXT("Aborted launch of worker: '%s'"), *TargetReceipt.Launch);
 						FPlatformProcess::TerminateProc(WorkerProcHandle, true);
 						FPlatformProcess::CloseProc(WorkerProcHandle);
-						return;
+						return false;
 					}
 				}
 
@@ -490,7 +478,7 @@ void FTargetReceiptBuildWorker::PopulateWorkerFromReceipt(const FString& TargetR
 				if (!FPaths::FileExists(FinalGeneratedVersionFilePath) || (FileManager.GetTimeStamp(*FinalGeneratedVersionFilePath) < NewestExecutableModificationTime))
 				{
 					UE_LOG(LogTargetReceiptBuildWorker, Warning, TEXT("Worker did not produce version file: '%s'"), *TargetReceipt.Launch);
-					return;
+					return false;
 				}
 			}
 
@@ -501,7 +489,7 @@ void FTargetReceiptBuildWorker::PopulateWorkerFromReceipt(const FString& TargetR
 			else
 			{
 				UE_LOG(LogTargetReceiptBuildWorker, Warning, TEXT("Worker did not produce readable version file: '%s'"), *TargetReceipt.Launch);
-				return;
+				return false;
 			}
 		}
 
@@ -509,7 +497,7 @@ void FTargetReceiptBuildWorker::PopulateWorkerFromReceipt(const FString& TargetR
 		if (!BuildSystemVersion.IsValid())
 		{
 			UE_LOG(LogTargetReceiptBuildWorker, Warning, TEXT("Invalid build system version in version file: '%s'"), *VersionFilePath);
-			return;
+			return false;
 		}
 
 		FCbArrayView FunctionsArrayView = VersionObject["Functions"].AsArrayView();
@@ -520,6 +508,79 @@ void FTargetReceiptBuildWorker::PopulateWorkerFromReceipt(const FString& TargetR
 		}
 	}
 
-	IModularFeatures::Get().RegisterModularFeature(UE::DerivedData::IBuildWorkerFactory::FeatureName, GetWorkerFactory());
-	bEnabled = true;
+	return true;
+}
+
+void FTargetReceiptBuildWorkerLoader::AbortLoad()
+{
+	bAbortRequested.store(true, std::memory_order_relaxed);
+}
+
+class FTargetReceiptBuildWorkerResolver final : public UE::DerivedData::IBuildWorkerResolver
+{
+public:
+	explicit FTargetReceiptBuildWorkerResolver(TMap<FIoHash, FString>&& InHashToPath)
+		: HashToPath(MoveTemp(InHashToPath))
+	{
+	}
+
+	void Resolve(TConstArrayView<FIoHash> RawHashes, UE::DerivedData::IRequestOwner& Owner, UE::DerivedData::FOnBuildWorkerResolved&& OnResolved) final
+	{
+		TArray<FCompressedBuffer> Files;
+		for (const FIoHash& RawHash : RawHashes)
+		{
+			if (const FString* Path = HashToPath.Find(RawHash))
+			{
+				if (TUniquePtr<FArchive> Ar{IFileManager::Get().CreateFileReader(**Path, FILEREAD_Silent)})
+				{
+					const int64 TotalSize = Ar->TotalSize();
+					FUniqueBuffer MutableBuffer = FUniqueBuffer::Alloc(uint64(TotalSize));
+					Ar->Serialize(MutableBuffer.GetData(), TotalSize);
+					if (Ar->Close())
+					{
+						Files.Emplace(FCompressedBuffer::Compress(
+							MutableBuffer.MoveToShared(),
+							ECompressedBufferCompressor::NotSet,
+							ECompressedBufferCompressionLevel::None));
+					}
+				}
+			}
+		}
+		OnResolved({Files, Files.Num() == RawHashes.Num() ? UE::DerivedData::EStatus::Ok : UE::DerivedData::EStatus::Error});
+	}
+
+private:
+	TMap<FIoHash, FString> HashToPath;
+};
+
+UE::DerivedData::FBuildWorker FTargetReceiptBuildWorkerLoader::Build(UE::DerivedData::FBuildWorkerBuilder& Builder) const
+{
+	Builder.SetName(Name);
+	Builder.SetPath(ExecutablePaths[0].RemotePath);
+	Builder.SetHostPlatform(Platform);
+
+	Builder.SetBuildSystemVersion(BuildSystemVersion);
+
+	for (const TPair<FUtf8String, FGuid>& FunctionVersion : FunctionVersions)
+	{
+		Builder.AddFunction(FunctionVersion.Key, FunctionVersion.Value);
+	}
+
+	for (int32 ExecutableIndex = 0; ExecutableIndex < ExecutablePaths.Num(); ++ExecutableIndex)
+	{
+		Builder.AddExecutable(ExecutablePaths[ExecutableIndex].RemotePath, ExecutableMeta[ExecutableIndex].Key, ExecutableMeta[ExecutableIndex].Value);
+	}
+
+	for (const TPair<FUtf8String, FUtf8String>& EnvironmentVariable : EnvironmentVariables)
+	{
+		Builder.SetEnvironment(EnvironmentVariable.Key, EnvironmentVariable.Value);
+	}
+
+	TMap<FIoHash, FString> HashToPath;
+	HashToPath.Reserve(ExecutablePaths.Num());
+	for (int32 ExecutableIndex = 0; ExecutableIndex < ExecutablePaths.Num(); ++ExecutableIndex)
+	{
+		HashToPath.Emplace(ExecutableMeta[ExecutableIndex].Get<FIoHash>(), ExecutablePaths[ExecutableIndex].LocalPath);
+	}
+	return Builder.Build(MakeUnique<FTargetReceiptBuildWorkerResolver>(MoveTemp(HashToPath)));
 }
diff --git a/Engine/Source/Developer/DesktopPlatform/Public/MultiPlatformTargetReceiptBuildWorkers.h b/Engine/Source/Developer/DesktopPlatform/Public/MultiPlatformTargetReceiptBuildWorkers.h
--- a/Engine/Source/Developer/DesktopPlatform/Public/MultiPlatformTargetReceiptBuildWorkers.h
+++ b/Engine/Source/Developer/DesktopPlatform/Public/MultiPlatformTargetReceiptBuildWorkers.h
@@ -2,25 +2,26 @@
 
 #pragma once
 
-#include "HAL/Platform.h"
+#include "Containers/StringFwd.h"
 #include "TargetReceiptBuildWorker.h"
 #include "Templates/TypeCompatibleBytes.h"
 
 #define UE_API DESKTOPPLATFORM_API
 
-class FTargetReceiptBuildWorker;
-
 /**
  * Globally registers a UE::DerivedData::IBuildWorkerFactory instance for each platform that build workers can be supported.
  * Users should include a $(Platform) variable in their file path to indicate that this worker receipt can exist for any platform.
  * If the receipt path does not include a $(Platform) variable, then no attempt will be made to find variations of it for other platforms.
-*/
-class FMultiPlatformTargetReceiptBuildWorkers
+ */
+class FMultiPlatformTargetReceiptBuildWorkers final
 {
 public:
-	UE_API FMultiPlatformTargetReceiptBuildWorkers(const TCHAR* TargetReceiptFilePath);
+	UE_API FMultiPlatformTargetReceiptBuildWorkers(FStringView TargetReceiptFilePath);
 	UE_API ~FMultiPlatformTargetReceiptBuildWorkers();
 
+	FMultiPlatformTargetReceiptBuildWorkers(const FMultiPlatformTargetReceiptBuildWorkers&) = delete;
+	FMultiPlatformTargetReceiptBuildWorkers& operator=(const FMultiPlatformTargetReceiptBuildWorkers&) = delete;
+
 private:
 	enum ESupportedPlatform : uint8
 	{
@@ -31,7 +32,7 @@ private:
 		SupportedPlatform_MAX
 	};
 	TTypeCompatibleBytes<FTargetReceiptBuildWorker> PlatformSpecificWorkerFactories[SupportedPlatform_MAX];
-	bool bAllPlatformsInitialized;
+	bool bAllPlatformsInitialized = false;
 };
 
 #undef UE_API
diff --git a/Engine/Source/Developer/DesktopPlatform/Public/TargetReceiptBuildWorker.h b/Engine/Source/Developer/DesktopPlatform/Public/TargetReceiptBuildWorker.h
--- a/Engine/Source/Developer/DesktopPlatform/Public/TargetReceiptBuildWorker.h
+++ b/Engine/Source/Developer/DesktopPlatform/Public/TargetReceiptBuildWorker.h
@@ -2,83 +2,28 @@
 
 #pragma once
 
-#include "Async/TaskGraphInterfaces.h"
-#include "Containers/Array.h"
-#include "Containers/Map.h"
 #include "Containers/StringFwd.h"
-#include "Containers/StringView.h"
-#include "Containers/UnrealString.h"
-#include "HAL/Platform.h"
-#include "IO/IoHash.h"
-#include "Misc/Guid.h"
-#include "Templates/Tuple.h"
-#include "Templates/TypeCompatibleBytes.h"
-
-#include <atomic>
+#include "Templates/PimplPtr.h"
 
 #define UE_API DESKTOPPLATFORM_API
 
 class FTargetReceiptBuildWorkerFactory;
-struct FIoHash;
-struct FTargetReceipt;
 
 /**
  * Globally registers a UE::DerivedData::IBuildWorkerFactory instance that runs an executable built by UnrealBuildTool.
  * UnrealBuildTool provides the executable information through a TargetReceipt file.
-*/
-class FTargetReceiptBuildWorker
+ */
+class FTargetReceiptBuildWorker final
 {
 public:
-	UE_API FTargetReceiptBuildWorker(const TCHAR* TargetReceiptFilePath);
-	UE_API virtual ~FTargetReceiptBuildWorker();
-
-private:
-	/**
-	 * Encapsulates a path used by the worker that has both a local and remote representation.
-	*/
-	class FWorkerPath
-	{
-	public:
-		FWorkerPath(FStringView InLocalPathRoot, FStringView InLocalPathSuffix, FStringView InRemotePathRoot = FString(), FStringView InRemotePathSuffix = FString());
-
-		FString GetLocalPath() const { return LocalPath; }
-		FString GetRemotePath() const { return RemotePath; }
-	private:
-		FString LocalPath;
-		FString RemotePath;
-	};
+	UE_API FTargetReceiptBuildWorker(FStringView TargetReceiptFilePath);
+	UE_API ~FTargetReceiptBuildWorker();
 
-	friend class FPopulateWorkerFromTargetReceiptTask;
-	friend class FTargetReceiptBuildWorkerFactory;
+	FTargetReceiptBuildWorker(const FTargetReceiptBuildWorker&) = delete;
+	FTargetReceiptBuildWorker& operator=(const FTargetReceiptBuildWorker&) = delete;
 
-	UE_API FTargetReceiptBuildWorkerFactory* GetWorkerFactory();
-	UE_API bool TryAddExecutablePath(FStringView Path);
-	UE_API void PopulateWorkerFromReceipt(const FString& TargetReceiptFilePath);
-
-	/**
-	 * Dummy type used to indicate the size of an empty class (no members) with virtual functions (contains a virtual function table).
-	 * The FTargetReceiptBuildWorkerFactory type will be created and stored in the space of one of these.
-	 * Size matching is enforced with a static_assert.  The purpose of this is to avoid exposing DerivedDataCache types to
-	 * modules that specify a target receipt worker.
-	*/
-	class IPureVirtual
-	{
-		virtual void Dummy();
-	};
-
-	static UE_API const FGuid WorkerReceiptVersion;
-	TTypeCompatibleBytes<IPureVirtual> InternalFactory;
-	FGraphEventRef PopulateTaskRef;
-	FString Name;
-	FString Platform;
-	TMap<FString, FString> EnvironmentVariables;
-	TArray<FWorkerPath> ExecutablePaths;
-	using FWorkerPathMeta = TTuple<FIoHash, uint64>;
-	TArray<FWorkerPathMeta> ExecutableMeta;
-	FGuid BuildSystemVersion;
-	TArray<TPair<FString, FGuid>> FunctionVersions;
-	std::atomic<bool> bAbortRequested;
-	bool bEnabled;
+private:
+	TPimplPtr<FTargetReceiptBuildWorkerFactory> Factory;
 };
 
 #undef UE_API
diff --git a/Engine/Source/Editor/UnrealEd/Private/DerivedDataBuildLocalExecutor.cpp b/Engine/Source/Editor/UnrealEd/Private/DerivedDataBuildLocalExecutor.cpp
--- a/Engine/Source/Editor/UnrealEd/Private/DerivedDataBuildLocalExecutor.cpp
+++ b/Engine/Source/Editor/UnrealEd/Private/DerivedDataBuildLocalExecutor.cpp
@@ -298,16 +298,16 @@ public:
 		// Calculate the worker hash that uniquely identifies the files of the worker.
 		FXxHash128Builder WorkerHashBuilder;
 		TArray<FIoHash> WorkerFileHashes;
-		TArray<TTuple<FStringView, bool>> WorkerFileMeta;
+		TArray<TTuple<FUtf8StringView, bool>> WorkerFileMeta;
 
-		Worker.IterateExecutables([&WorkerHashBuilder, &WorkerFileHashes, &WorkerFileMeta](FStringView Path, const FIoHash& RawHash, uint64 RawSize)
+		Worker.IterateExecutables([&WorkerHashBuilder, &WorkerFileHashes, &WorkerFileMeta](FUtf8StringView Path, const FIoHash& RawHash, uint64 RawSize)
 		{
 			WorkerHashBuilder.Update(MakeMemoryView(&RawHash, 1));
 			WorkerFileHashes.Emplace(RawHash);
 			WorkerFileMeta.Emplace(Path, true);
 		});
 
-		Worker.IterateFiles([&WorkerHashBuilder, &WorkerFileHashes, &WorkerFileMeta](FStringView Path, const FIoHash& RawHash, uint64 RawSize)
+		Worker.IterateFiles([&WorkerHashBuilder, &WorkerFileHashes, &WorkerFileMeta](FUtf8StringView Path, const FIoHash& RawHash, uint64 RawSize)
 		{
 			WorkerHashBuilder.Update(MakeMemoryView(&RawHash, 1));
 			WorkerFileHashes.Emplace(RawHash);
@@ -366,7 +366,7 @@ public:
 
 		std::atomic<bool> bOk = true;
 		FRequestOwner BlockingOwner(EPriority::Blocking);
-		Worker.FindFileData(WorkerFileHashes, BlockingOwner, [&Worker, &FileManager, &ScratchPath, &WorkerFileMeta, &bOk](FBuildWorkerFileDataCompleteParams&& Params)
+		Worker.Resolve(WorkerFileHashes, BlockingOwner, [&Worker, &FileManager, &ScratchPath, &WorkerFileMeta, &bOk](FBuildWorkerResolvedParams&& Params)
 		{
 			uint32 MetaIndex = 0;
 			for (const FCompressedBuffer& Buffer : Params.Files)
@@ -423,9 +423,9 @@ public:
 		return bExists;
 	}
 
-	TConstArrayView<FStringView> GetHostPlatforms() const final
+	TConstArrayView<FUtf8StringView> GetHostPlatforms() const final
 	{
-		static constexpr FStringView HostPlatforms[]{TEXTVIEW("Win64")};
+		static constexpr FUtf8StringView HostPlatforms[]{UTF8TEXTVIEW("Win64")};
 		return HostPlatforms;
 	}
 
diff --git a/Engine/Source/Editor/UnrealEd/Private/DerivedDataBuildRemoteExecutor.cpp b/Engine/Source/Editor/UnrealEd/Private/DerivedDataBuildRemoteExecutor.cpp
--- a/Engine/Source/Editor/UnrealEd/Private/DerivedDataBuildRemoteExecutor.cpp
+++ b/Engine/Source/Editor/UnrealEd/Private/DerivedDataBuildRemoteExecutor.cpp
@@ -239,9 +239,9 @@ public:
 		new FRemoteBuildExecutionRequest(*this, Action, Inputs, Policy, Worker, BuildSystem, Owner, MoveTemp(OnComplete));
 	}
 
-	TConstArrayView<FStringView> GetHostPlatforms() const final
+	TConstArrayView<FUtf8StringView> GetHostPlatforms() const final
 	{
-		static constexpr FStringView HostPlatforms[]{TEXTVIEW("Win64"), TEXTVIEW("Linux"), TEXTVIEW("Mac")};
+		static constexpr FUtf8StringView HostPlatforms[]{UTF8TEXTVIEW("Win64"), UTF8TEXTVIEW("Linux"), UTF8TEXTVIEW("Mac")};
 		return HostPlatforms;
 	}
 
@@ -518,14 +518,14 @@ FCbObject FRemoteBuildExecutionRequest::BuildWorkerDescriptor(const FBuildWorker
 	//WorkerDescriptor.AddInteger(ANSITEXTVIEW("memory"), 1 * 1024 * 1024 * 1024);
 
 	WorkerDescriptor.BeginArray(ANSITEXTVIEW("environment"));
-	Worker.IterateEnvironment([&WorkerDescriptor](FStringView Name, FStringView Value)
+	Worker.IterateEnvironment([&WorkerDescriptor](FUtf8StringView Name, FUtf8StringView Value)
 		{
-			WorkerDescriptor.AddString(WriteToString<256>(Name, "=", Value));
+			WorkerDescriptor.AddString(WriteToUtf8String<256>(Name, "=", Value));
 		});
 	WorkerDescriptor.EndArray();
 
 	WorkerDescriptor.BeginArray(ANSITEXTVIEW("executables"));
-	Worker.IterateExecutables([&WorkerDescriptor](FStringView Key, const FIoHash& RawHash, uint64 RawSize)
+	Worker.IterateExecutables([&WorkerDescriptor](FUtf8StringView Key, const FIoHash& RawHash, uint64 RawSize)
 		{
 			WorkerDescriptor.BeginObject();
 			WorkerDescriptor.AddString(ANSITEXTVIEW("name"), Key);
@@ -536,7 +536,7 @@ FCbObject FRemoteBuildExecutionRequest::BuildWorkerDescriptor(const FBuildWorker
 	WorkerDescriptor.EndArray();
 
 	WorkerDescriptor.BeginArray("files");
-	Worker.IterateFiles([&WorkerDescriptor](FStringView Key, const FIoHash& RawHash, uint64 RawSize)
+	Worker.IterateFiles([&WorkerDescriptor](FUtf8StringView Key, const FIoHash& RawHash, uint64 RawSize)
 		{
 			WorkerDescriptor.BeginObject();
 			WorkerDescriptor.AddString(ANSITEXTVIEW("name"), Key);
@@ -676,7 +676,7 @@ void FRemoteBuildExecutionRequest::PostWorkerPackage_Async()
 
 			{
 				TSet<FIoHash> WorkerFileHashes;
-				State.BuildWorker.IterateExecutables([NeedHashes = State.NeedHashes, &WorkerFileHashes](FStringView Path, const FIoHash& RawHash, uint64 RawSize)
+				State.BuildWorker.IterateExecutables([NeedHashes = State.NeedHashes, &WorkerFileHashes](FUtf8StringView Path, const FIoHash& RawHash, uint64 RawSize)
 					{
 						if (NeedHashes.Contains(RawHash))
 						{
@@ -684,7 +684,7 @@ void FRemoteBuildExecutionRequest::PostWorkerPackage_Async()
 						}
 					});
 
-				State.BuildWorker.IterateFiles([NeedHashes = State.NeedHashes, &WorkerFileHashes](FStringView Path, const FIoHash& RawHash, uint64 RawSize)
+				State.BuildWorker.IterateFiles([NeedHashes = State.NeedHashes, &WorkerFileHashes](FUtf8StringView Path, const FIoHash& RawHash, uint64 RawSize)
 					{
 						if (NeedHashes.Contains(RawHash))
 						{
@@ -702,7 +702,7 @@ void FRemoteBuildExecutionRequest::PostWorkerPackage_Async()
 				}
 
 				FRequestOwner BlockingOwner(EPriority::Blocking);
-				State.BuildWorker.FindFileData(WorkerFileHashes.Array(), BlockingOwner, [&Package, &AttachmentBytes](FBuildWorkerFileDataCompleteParams&& Params)
+				State.BuildWorker.Resolve(WorkerFileHashes.Array(), BlockingOwner, [&Package, &AttachmentBytes](FBuildWorkerResolvedParams&& Params)
 					{
 						for (const FCompressedBuffer& Buffer : Params.Files)
 						{
