DDB: Added packaging of build workers

- Added a unique key to build workers to identify them.
- Added a list of required directories to build workers to model the required binaries directory more generically.
- Building a build worker now requires the file data to be passed to the builder. Compression is determined by the caller.
- Added GetPackagePath() to IBuildWorkerResolver to allow direct package access for resolvers that are based on a package.
- Extended IBuildWorkerFactory to more explicitly support asynchronous creation of build workers and multiple workers from one factory.
- Extended IBuildWorkerRegistry with CreateWorker() and LoadWorker() to support interface-only (no imported symbols) creation and loading of workers.
- Creating a build worker from a target receipt now creates a compressed package on disk that the build worker can be served from.
- Added CreateBuildWorkerFromTargetReceipt with parameters to force a build, to include symbols, and to skip compression.
- Fixed the intrusive unset optional implementation for FBuildWorker.

#Synced-CL 50253651

diff --git a/Engine/Source/Developer/DerivedDataCache/Private/DerivedDataBuildWorker.cpp b/Engine/Source/Developer/DerivedDataCache/Private/DerivedDataBuildWorker.cpp
new file mode 100644
--- /dev/null
+++ b/Engine/Source/Developer/DerivedDataCache/Private/DerivedDataBuildWorker.cpp
@@ -0,0 +1,667 @@
+// Copyright Epic Games, Inc. All Rights Reserved.
+
+#include "DerivedDataBuildWorker.h"
+
+#include "Algo/Sort.h"
+#include "Algo/Unique.h"
+#include "Compression/CompressedBuffer.h"
+#include "Containers/Array.h"
+#include "Containers/SharedString.h"
+#include "Containers/Utf8String.h"
+#include "DerivedDataBuildPrivate.h"
+#include "DerivedDataRequestOwner.h"
+#include "HAL/FileManager.h"
+#include "HAL/PlatformFileManager.h"
+#include "HAL/PlatformMath.h"
+#include "Hash/xxhash.h"
+#include "Logging/StructuredLog.h"
+#include "Misc/Guid.h"
+#include "Misc/PathViews.h"
+#include "Serialization/CompactBinaryWriter.h"
+#include "Templates/AlignmentTemplates.h"
+#include "Templates/Function.h"
+#include "Templates/Tuple.h"
+
+namespace UE::DerivedData::Private
+{
+
+struct FBuildWorkerPackageAttachment
+{
+	FIoHash RawHash;
+	uint32 RawSizeLo = 0;
+	uint16 RawSizeHi = 0;
+	uint16 CompressedSizeLo = 0;
+	uint32 CompressedSizeHi = 0;
+	uint64 CompressedOffset = 0;
+
+	inline uint64 GetRawSize() const 
+	{
+		return (uint64(RawSizeHi) << 32) | RawSizeLo;
+	}
+
+	inline void SetRawSize(uint64 Size)
+	{
+		RawSizeHi = uint16(Size >> 32);
+		RawSizeLo = uint32(Size);
+	}
+
+	inline uint64 GetCompressedSize() const 
+	{
+		return (uint64(CompressedSizeHi) << 16) | CompressedSizeLo;
+	}
+
+	inline void SetCompressedSize(uint64 Size)
+	{
+		CompressedSizeHi = uint32(Size >> 16);
+		CompressedSizeLo = uint16(Size);
+	}
+};
+static_assert(sizeof(FBuildWorkerPackageAttachment) == 40);
+
+struct FBuildWorkerPackageHeader
+{
+	constexpr static uint32 ExpectedMagic = 0xb7756277; // <dot>ubw
+
+	/** A magic number to identify a build worker package. Always 0xb7756277. */
+	uint32 Magic = ExpectedMagic;
+	/** The hash of the header from AttachmentCount to the end of the attachment table. */
+	uint32 HeaderHash = 0;
+	/** The number of attachments in this package. */
+	uint32 AttachmentCount = 0;
+	/** The offset of the attachment data relative to the start of this header. */
+	uint32 AttachmentOffset = 0;
+	/** The total size of the worker object following the header. */
+	uint32 WorkerSize = 0;
+	/** The hash of the worker object following the header. */
+	FIoHash WorkerHash;
+	/** The table of attachments in the package. Sorted by RawHash. Size is AttachmentCount. */
+	FBuildWorkerPackageAttachment Attachments[0];
+
+	/** Calculate HeaderHash using the low 32 bits from FXxHash64. */
+	uint32 CalculateHash() const
+	{
+		const uint32 HeaderSize = sizeof(*this) + sizeof(*Attachments) * AttachmentCount;
+		const uint32 HashOffset = offsetof(FBuildWorkerPackageHeader, AttachmentCount);
+		const FMemoryView HeaderView = MakeMemoryView(this, HeaderSize);
+		return (uint32)FXxHash64::HashBuffer(HeaderView.RightChop(HashOffset)).Hash;
+	}
+};
+
+class FBuildWorkerPackageResolver final : public IBuildWorkerResolver
+{
+public:
+	explicit FBuildWorkerPackageResolver(TNotNull<FBuildWorkerPackageHeader*> Header, FStringView PackagePath);
+	~FBuildWorkerPackageResolver();
+
+	FBuildWorkerPackageResolver(const FBuildWorkerPackageResolver&) = delete;
+	FBuildWorkerPackageResolver& operator=(const FBuildWorkerPackageResolver&) = delete;
+
+	void Resolve(TConstArrayView<FIoHash> RawHashes, IRequestOwner& Owner, FOnBuildWorkerResolved&& OnResolved) final;
+
+	[[nodiscard]] FStringView GetPackagePath() const
+	{
+		return PackagePath;
+	}
+
+private:
+	TNotNull<FBuildWorkerPackageHeader*> Header;
+	FSharedString PackagePath;
+};
+
+FBuildWorkerPackageResolver::FBuildWorkerPackageResolver(TNotNull<FBuildWorkerPackageHeader*> InHeader, FStringView InPackagePath)
+	: Header(InHeader)
+	, PackagePath(InPackagePath)
+{
+}
+
+FBuildWorkerPackageResolver::~FBuildWorkerPackageResolver()
+{
+	static_assert(std::is_trivially_destructible_v<FBuildWorkerPackageHeader>);
+	FMemory::Free(Header);
+}
+
+void FBuildWorkerPackageResolver::Resolve(TConstArrayView<FIoHash> RawHashes, IRequestOwner& Owner, FOnBuildWorkerResolved&& OnResolved)
+{
+	TConstArrayView<FBuildWorkerPackageAttachment> Attachments(Header->Attachments, (int32)Header->AttachmentCount);
+	TArray<FCompressedBuffer> ResolvedData;
+
+	if (!RawHashes.IsEmpty())
+	{
+		if (TUniquePtr<FArchive> PackageAr{IFileManager::Get().CreateFileReader(*PackagePath, FILEREAD_Silent)})
+		{
+			const uint64 PackageSize = (uint64)FPlatformMath::Max(0, PackageAr->TotalSize());
+			ResolvedData.Reserve(RawHashes.Num());
+			for (const FIoHash& RawHash : RawHashes)
+			{
+				const int32 AttachmentIndex = Algo::BinarySearchBy(Attachments, RawHash, &FBuildWorkerPackageAttachment::RawHash);
+				if (AttachmentIndex >= 0)
+				{
+					const FBuildWorkerPackageAttachment& Attachment = Attachments[AttachmentIndex];
+					const uint64 CompressedOffset = Attachment.CompressedOffset + Header->AttachmentOffset;
+					const uint64 CompressedSize = Attachment.GetCompressedSize();
+					const uint64 CompressedEndOffset = CompressedOffset + CompressedSize;
+					if (CompressedOffset <= CompressedEndOffset && CompressedEndOffset <= PackageSize)
+					{
+						FUniqueBuffer CompressedData = FUniqueBuffer::Alloc(CompressedSize);
+						PackageAr->Seek((int64)CompressedOffset);
+						PackageAr->Serialize(CompressedData.GetData(), (int64)CompressedSize);
+						FCompressedBuffer Data = FCompressedBuffer::FromCompressed(CompressedData.MoveToShared());
+						if (Data.GetRawHash() == RawHash)
+						{
+							ResolvedData.Add(MoveTemp(Data));
+						}
+					}
+				}
+			}
+		}
+	}
+
+	const EStatus Status = ResolvedData.Num() == RawHashes.Num() ? EStatus::Ok : EStatus::Error;
+	OnResolved({ResolvedData, Status});
+}
+
+} // UE::DerivedData::Private
+
+///////////////////////////////////////////////////////////////////////////////////////////////////////////////////////
+
+namespace UE::DerivedData
+{
+
+FUtf8StringView FBuildWorker::GetName() const
+{
+	return Object.FindView(ANSITEXTVIEW("Name")).AsString();
+}
+
+FUtf8StringView FBuildWorker::GetPath() const
+{
+	return Object.FindView(ANSITEXTVIEW("Path")).AsString();
+}
+
+FUtf8StringView FBuildWorker::GetHostPlatform() const
+{
+	return Object.FindView(ANSITEXTVIEW("HostPlatform")).AsString();
+}
+
+FGuid FBuildWorker::GetBuildSystemVersion() const
+{
+	return Object.FindView(ANSITEXTVIEW("BuildSystemVersion")).AsUuid();
+}
+
+void FBuildWorker::IterateFunctions(TFunctionRef<void (FUtf8StringView Name, const FGuid& Version)> Visitor) const
+{
+	for (FCbFieldView It : Object.FindView(ANSITEXTVIEW("Functions")))
+	{
+		Visitor(It[ANSITEXTVIEW("Name")].AsString(), It[ANSITEXTVIEW("Version")].AsUuid());
+	}
+}
+
+void FBuildWorker::IterateDirectories(TFunctionRef<void (FUtf8StringView Path)> Visitor) const
+{
+	for (FCbFieldView It : Object.FindView(ANSITEXTVIEW("Directories")))
+	{
+		Visitor(It[ANSITEXTVIEW("Path")].AsString());
+	}
+}
+
+void FBuildWorker::IterateFiles(TFunctionRef<void (FUtf8StringView Path, const FIoHash& RawHash, uint64 RawSize)> Visitor) const
+{
+	for (FCbFieldView It : Object.FindView(ANSITEXTVIEW("Files")))
+	{
+		Visitor(It[ANSITEXTVIEW("Path")].AsString(), It[ANSITEXTVIEW("RawHash")].AsHash(), It[ANSITEXTVIEW("RawSize")].AsUInt64());
+	}
+}
+
+void FBuildWorker::IterateExecutables(TFunctionRef<void (FUtf8StringView Path, const FIoHash& RawHash, uint64 RawSize)> Visitor) const
+{
+	for (FCbFieldView It : Object.FindView(ANSITEXTVIEW("Executables")))
+	{
+		Visitor(It[ANSITEXTVIEW("Path")].AsString(), It[ANSITEXTVIEW("RawHash")].AsHash(), It[ANSITEXTVIEW("RawSize")].AsUInt64());
+	}
+}
+
+void FBuildWorker::IterateEnvironment(TFunctionRef<void (FUtf8StringView Name, FUtf8StringView Value)> Visitor) const
+{
+	for (FCbFieldView It : Object.FindView(ANSITEXTVIEW("Environment")))
+	{
+		Visitor(It[ANSITEXTVIEW("Name")].AsString(), It[ANSITEXTVIEW("Value")].AsString());
+	}
+}
+
+void FBuildWorker::Resolve(TConstArrayView<FIoHash> RawHashes, IRequestOwner& Owner, FOnBuildWorkerResolved&& OnResolved) const
+{
+	Resolver->Resolve(RawHashes, Owner, MoveTemp(OnResolved));
+}
+
+bool FBuildWorker::Extract(FStringView OutputPathView) const
+{
+	std::atomic<bool> bOk = true;
+	IFileManager& FileManager = IFileManager::Get();
+
+	// Extract the worker to a temporary directory.
+	TStringBuilder<256> OutputPath(InPlace, OutputPathView);
+	TStringBuilder<256> ScratchPath(InPlace, OutputPath, '.', FGuid::NewGuid());
+
+	// Create the directories for the build worker.
+	IterateDirectories([this, &FileManager, &ScratchPath, &bOk](FUtf8StringView Path)
+	{
+		TStringBuilder<256> DirectoryPath;
+		FPathViews::Append(DirectoryPath, ScratchPath, Path);
+		if (!FileManager.MakeDirectory(*DirectoryPath, /*bTree*/ true))
+		{
+			UE_LOGFMT(LogDerivedDataBuild, Display,
+				"Failed to create directory '{Path}' for build worker {WorkerName}.", DirectoryPath, GetName());
+			bOk.store(false, std::memory_order_relaxed);
+		}
+	});
+
+	// Gather the files that need to be manifested.
+	TArray<FIoHash> WorkerFileHashes;
+	TArray<TTuple<FUtf8StringView, bool>> WorkerFileMeta;
+
+	IterateExecutables([&WorkerFileHashes, &WorkerFileMeta](FUtf8StringView Path, const FIoHash& RawHash, uint64 RawSize)
+	{
+		WorkerFileHashes.Emplace(RawHash);
+		WorkerFileMeta.Emplace(Path, true);
+	});
+
+	IterateFiles([&WorkerFileHashes, &WorkerFileMeta](FUtf8StringView Path, const FIoHash& RawHash, uint64 RawSize)
+	{
+		WorkerFileHashes.Emplace(RawHash);
+		WorkerFileMeta.Emplace(Path, false);
+	});
+
+	// Create the files for the build worker.
+	FRequestOwner BlockingOwner(EPriority::Blocking);
+	Resolve(WorkerFileHashes, BlockingOwner, [this, &FileManager, &ScratchPath, &WorkerFileMeta, &bOk](FBuildWorkerResolvedParams&& Params)
+	{
+		if (Params.Status != EStatus::Ok)
+		{
+			bOk.store(false, std::memory_order_relaxed);
+		}
+		uint32 MetaIndex = 0;
+		for (const FCompressedBuffer& CompressedData : Params.Files)
+		{
+			TStringBuilder<256> FilePath;
+			FPathViews::Append(FilePath, ScratchPath, WorkerFileMeta[MetaIndex].Key);
+			if (TUniquePtr<FArchive> Ar{FileManager.CreateFileWriter(*FilePath, FILEWRITE_Silent)})
+			{
+				if (FCompositeBuffer RawData = CompressedData.DecompressToComposite())
+				{
+					for (auto& Segment : RawData.GetSegments())
+					{
+						Ar->Serialize((void*)Segment.GetData(), Segment.GetSize());
+					}
+				}
+				else
+				{
+					UE_LOGFMT(LogDerivedDataBuild, Display,
+						"Failed to decompress data for '{Path}' for build worker {WorkerName}.", FilePath, GetName());
+					bOk.store(false, std::memory_order_relaxed);
+				}
+			}
+			else
+			{
+				UE_LOGFMT(LogDerivedDataBuild, Display,
+					"Failed to open '{Path}' for writing for build worker {WorkerName}.", FilePath, GetName());
+				bOk.store(false, std::memory_order_relaxed);
+			}
+
+			++MetaIndex;
+		}
+	});
+	BlockingOwner.Wait();
+
+	// Move the temporary directory to its final location or delete it if there was a race.
+	IPlatformFile& PlatformFile = FPlatformFileManager::Get().GetPlatformFile();
+	if (bOk.load(std::memory_order_relaxed) && !PlatformFile.DirectoryExists(*OutputPath))
+	{
+		PlatformFile.MoveFile(*OutputPath, *ScratchPath);
+	}
+	PlatformFile.DeleteDirectoryRecursively(*ScratchPath);
+	const bool bExists = PlatformFile.DirectoryExists(*OutputPath);
+	UE_CLOGFMT(!bExists, LogDerivedDataBuild, Display,
+		"Failed to rename '{TempPath}' to '{Path}' for build worker {WorkerName}.", ScratchPath, OutputPath, GetName());
+
+	return bExists;
+}
+
+TOptional<FBuildWorker> FBuildWorker::Load(FStringView PackagePath)
+{
+	using namespace UE::DerivedData::Private;
+	TUniquePtr<FArchive> PackageAr(IFileManager::Get().CreateFileReader(*WriteToString<256>(PackagePath), FILEREAD_Silent));
+	if (!PackageAr)
+	{
+		return {};
+	}
+
+	const int64 TotalSize = PackageAr->TotalSize();
+
+	FBuildWorkerPackageHeader HeaderOnly{};
+	if (TotalSize < sizeof(HeaderOnly))
+	{
+		UE_LOGFMT(LogDerivedDataBuild, Warning, "Failed to load build worker package '{Package}' because "
+			"the file size {FileSize} is below the minimum header size {HeaderSize}.",
+			PackagePath, TotalSize, sizeof(HeaderOnly));
+		return {};
+	}
+
+	PackageAr->Serialize(&HeaderOnly, sizeof(HeaderOnly));
+
+	if (HeaderOnly.Magic != HeaderOnly.ExpectedMagic)
+	{
+		UE_LOGFMT(LogDerivedDataBuild, Warning, "Failed to load build worker package '{Package}' because "
+			"the magic number in the header is incorrect.", PackagePath);
+		return {};
+	}
+
+	const int64 HeaderSize = sizeof(HeaderOnly) + sizeof(*HeaderOnly.Attachments) * (int64)HeaderOnly.AttachmentCount;
+	if (TotalSize < HeaderSize)
+	{
+		UE_LOGFMT(LogDerivedDataBuild, Warning, "Failed to load build worker package '{Package}' because "
+			"the file size {FileSize} is below the reported header size {HeaderSize}.",
+			PackagePath, TotalSize, HeaderSize);
+		return {};
+	}
+
+	struct FDeleteByFree
+	{
+		void operator()(void* Ptr) const
+		{
+			FMemory::Free(Ptr);
+		}
+	};
+	TUniquePtr<FBuildWorkerPackageHeader, FDeleteByFree> Header(new(FMemory::Malloc(HeaderSize, alignof(FBuildWorkerPackageHeader))) FBuildWorkerPackageHeader(HeaderOnly));
+
+	PackageAr->Serialize(Header->Attachments, HeaderSize - sizeof(HeaderOnly));
+
+	if (Header->CalculateHash() != Header->HeaderHash)
+	{
+		UE_LOGFMT(LogDerivedDataBuild, Warning, "Failed to load build worker package '{Package}' because "
+			"the header hash does not match.", PackagePath);
+		return {};
+	}
+
+	int64 PackageSize = HeaderSize + Header->WorkerSize;
+	if (Header->AttachmentCount)
+	{
+		if (Header->AttachmentOffset < PackageSize)
+		{
+			UE_LOGFMT(LogDerivedDataBuild, Warning, "Failed to load build worker package '{Package}' because "
+				"the attachment offset {AttachmentOffset} is before the end of the header and worker object at {PackageSize}.",
+				PackagePath, Header->AttachmentOffset, PackageSize);
+			return {};
+		}
+		const FBuildWorkerPackageAttachment& LastAttachment = Header->Attachments[Header->AttachmentCount - 1];
+		PackageSize = Header->AttachmentOffset + LastAttachment.CompressedOffset + (int64)LastAttachment.GetCompressedSize();
+	}
+	if (TotalSize < PackageSize)
+	{
+		UE_LOGFMT(LogDerivedDataBuild, Warning, "Failed to load build worker package '{Package}' because "
+			"the file size {FileSize} is below the reported package size {PackageSize}.",
+			PackagePath, TotalSize, PackageSize);
+		return {};
+	}
+
+	FUniqueBuffer WorkerObjectData = FUniqueBuffer::Alloc(Header->WorkerSize);
+	PackageAr->Serialize(WorkerObjectData.GetData(), Header->WorkerSize);
+	if (FIoHash::HashBuffer(WorkerObjectData) != Header->WorkerHash)
+	{
+		UE_LOGFMT(LogDerivedDataBuild, Warning, "Failed to load build worker package '{Package}' because "
+			"the worker hash does not match. Package has been corrupted.", PackagePath);
+		return {};
+	}
+
+	FCbObject WorkerObject(WorkerObjectData.MoveToShared());
+	FBuildWorker Worker;
+	Worker.Key = {Header->WorkerHash};
+	Worker.Object = MoveTemp(WorkerObject);
+	Worker.Resolver = MakeUnique<FBuildWorkerPackageResolver>(Header.Release(), PackagePath);
+	return Worker;
+}
+
+FBuildWorker FBuildWorker::Load(FCbObject Object, TNotNull<TUniquePtr<IBuildWorkerResolver>> Resolver)
+{
+	Object.MakeOwned();
+	FBuildWorker Worker;
+	Worker.Object = MoveTemp(Object);
+	Worker.Resolver = MoveTemp(Resolver);
+	Worker.Key = {Worker.Object.GetHash()};
+	return Worker;
+}
+
+} // UE::DerivedData
+
+///////////////////////////////////////////////////////////////////////////////////////////////////////////////////////
+
+namespace UE::DerivedData::Private
+{
+
+class FBuildWorkerBuilderInternal final : public IBuildWorkerBuilderInternal
+{
+public:
+	void SetName(FUtf8StringView Name) final
+	{
+		WorkerName = Name;
+	}
+
+	void SetPath(FUtf8StringView Path) final
+	{
+		WorkerPath = Path;
+	}
+
+	void SetHostPlatform(FUtf8StringView Name) final
+	{
+		HostPlatform = Name;
+	}
+
+	void SetBuildSystemVersion(const FGuid& Version) final
+	{
+		BuildSystemVersion = Version;
+	}
+
+	void AddDirectory(FUtf8StringView Path) final
+	{
+		Directories.Emplace(Path);
+	}
+
+	void AddFunction(FUtf8StringView Name, const FGuid& Version) final
+	{
+		Functions.Emplace(Name, Version);
+	}
+
+	void AddFile(FUtf8StringView Path, const FCompressedBuffer& Data) final
+	{
+		Files.Emplace(Path, Data);
+	}
+
+	void AddExecutable(FUtf8StringView Path, const FCompressedBuffer& Data) final
+	{
+		Executables.Emplace(Path, Data);
+	}
+
+	void SetEnvironment(FUtf8StringView Name, FUtf8StringView Value) final
+	{
+		Environment.Emplace(Name, Value);
+	}
+
+	FBuildWorker Build(FStringView PackagePath) final;
+
+private:
+	FUtf8String WorkerName;
+	FUtf8String WorkerPath;
+	FUtf8String HostPlatform;
+	FGuid BuildSystemVersion;
+	TArray<TTuple<FUtf8String, FGuid>> Functions;
+	TArray<FUtf8String> Directories;
+	TArray<TTuple<FUtf8String, FCompressedBuffer>> Files;
+	TArray<TTuple<FUtf8String, FCompressedBuffer>> Executables;
+	TArray<TTuple<FUtf8String, FUtf8String>> Environment;
+};
+
+FBuildWorker FBuildWorkerBuilderInternal::Build(FStringView PackagePath)
+{
+	constexpr static uint32 AttachmentAlignment = 64;
+
+	TArray<FCompressedBuffer> Attachments;
+
+	// Create the worker object and gather attachments.
+	Algo::Sort(Functions);
+	Algo::Sort(Directories);
+	Algo::SortBy(Files, &TTuple<FUtf8String, FCompressedBuffer>::Key);
+	Algo::SortBy(Executables, &TTuple<FUtf8String, FCompressedBuffer>::Key);
+	Algo::Sort(Environment);
+
+	TCbWriter<4096> Writer;
+	Writer.BeginObject();
+
+	Writer.AddString(ANSITEXTVIEW("Name"), WorkerName);
+	Writer.AddString(ANSITEXTVIEW("Path"), WorkerPath);
+	Writer.AddString(ANSITEXTVIEW("HostPlatform"), HostPlatform);
+	Writer.AddUuid(ANSITEXTVIEW("BuildSystemVersion"), BuildSystemVersion);
+
+	if (!Functions.IsEmpty())
+	{
+		Writer.BeginArray(ANSITEXTVIEW("Functions"));
+		for (const TTuple<FUtf8String, FGuid>& Function : Functions)
+		{
+			const FUtf8String& Name = Function.Get<FUtf8String>();
+			const FGuid& Version = Function.Get<FGuid>();
+			UE_CLOGF(!Version.IsValid(), LogDerivedDataBuild, Error,
+				"Version of zero is not allowed in build function with the name %s in build worker '%s'.",
+				*WriteToUtf8String<32>(Name), *WorkerName);
+			Writer.BeginObject();
+			Writer.AddString(ANSITEXTVIEW("Name"), Name);
+			Writer.AddUuid(ANSITEXTVIEW("Version"), Version);
+			Writer.EndObject();
+		}
+		Writer.EndArray();
+	}
+
+	if (!Directories.IsEmpty())
+	{
+		Writer.BeginArray(ANSITEXTVIEW("Directories"));
+		for (const FUtf8String& Directory : Directories)
+		{
+			Writer.BeginObject();
+			Writer.AddString(ANSITEXTVIEW("Path"), Directory);
+			Writer.EndObject();
+		}
+		Writer.EndArray();
+	}
+
+	if (!Files.IsEmpty())
+	{
+		Writer.BeginArray(ANSITEXTVIEW("Files"));
+		for (const TTuple<FUtf8String, FCompressedBuffer>& File : Files)
+		{
+			const FCompressedBuffer& Attachment = File.Get<FCompressedBuffer>();
+			Attachments.Add(Attachment);
+			Writer.BeginObject();
+			Writer.AddString(ANSITEXTVIEW("Path"), File.Get<FUtf8String>());
+			Writer.AddBinaryAttachment(ANSITEXTVIEW("RawHash"), Attachment.GetRawHash());
+			Writer.AddInteger(ANSITEXTVIEW("RawSize"), Attachment.GetRawSize());
+			Writer.EndObject();
+		}
+		Writer.EndArray();
+	}
+
+	if (!Executables.IsEmpty())
+	{
+		Writer.BeginArray(ANSITEXTVIEW("Executables"));
+		for (const TTuple<FUtf8String, FCompressedBuffer>& Executable : Executables)
+		{
+			const FCompressedBuffer& Attachment = Executable.Get<FCompressedBuffer>();
+			Attachments.Add(Attachment);
+			Writer.BeginObject();
+			Writer.AddString(ANSITEXTVIEW("Path"), Executable.Get<FUtf8String>());
+			Writer.AddBinaryAttachment(ANSITEXTVIEW("RawHash"), Attachment.GetRawHash());
+			Writer.AddInteger(ANSITEXTVIEW("RawSize"), Attachment.GetRawSize());
+			Writer.EndObject();
+		}
+		Writer.EndArray();
+	}
+
+	if (!Environment.IsEmpty())
+	{
+		Writer.BeginArray(ANSITEXTVIEW("Environment"));
+		for (const TTuple<FUtf8String, FUtf8String>& Var : Environment)
+		{
+			Writer.BeginObject();
+			Writer.AddString(ANSITEXTVIEW("Name"), Var.Get<0>());
+			Writer.AddString(ANSITEXTVIEW("Value"), Var.Get<1>());
+			Writer.EndObject();
+		}
+		Writer.EndArray();
+	}
+
+	Writer.EndObject();
+	FCbObject WorkerObject = Writer.Save().AsObject();
+
+	// Build the worker package header from the sorted attachments.
+	Algo::SortBy(Attachments, &FCompressedBuffer::GetRawHash);
+	Attachments.SetNum(Algo::UniqueBy(Attachments, &FCompressedBuffer::GetRawHash), EAllowShrinking::No);
+
+	const uint32 HeaderSize = sizeof(FBuildWorkerPackageHeader) + sizeof(FBuildWorkerPackageAttachment) * Attachments.Num();
+	FBuildWorkerPackageHeader* Header = new(FMemory::Malloc(HeaderSize, alignof(FBuildWorkerPackageHeader))) FBuildWorkerPackageHeader;
+	Header->AttachmentCount = (uint32)Attachments.Num();
+
+	uint64 CompressedOffset = 0;
+	for (int32 AttachmentIndex = 0, AttachmentCount = Attachments.Num(); AttachmentIndex < AttachmentCount; ++AttachmentIndex)
+	{
+		const FCompressedBuffer& Attachment = Attachments[AttachmentIndex];
+		const uint64 CompressedSize = Attachment.GetCompressedSize();
+		FBuildWorkerPackageAttachment& AttachmentHeader = Header->Attachments[AttachmentIndex];
+		AttachmentHeader.RawHash = Attachment.GetRawHash();
+		AttachmentHeader.SetRawSize(Attachment.GetRawSize());
+		AttachmentHeader.SetCompressedSize(CompressedSize);
+		AttachmentHeader.CompressedOffset = CompressedOffset;
+		CompressedOffset = Align(CompressedOffset + CompressedSize, AttachmentAlignment);
+	}
+
+	// Create the worker and its contained worker object.
+	FBuildWorker Worker = FBuildWorker::Load(WorkerObject, MakeUnique<FBuildWorkerPackageResolver>(Header, PackagePath));
+
+	Header->WorkerSize = (uint32)WorkerObject.GetSize();
+	Header->WorkerHash = Worker.GetKey().Hash;
+	Header->AttachmentOffset = Align(HeaderSize + Header->WorkerSize, AttachmentAlignment);
+
+	Header->HeaderHash = Header->CalculateHash();
+
+	// Save the worker package to disk.
+	uint8 AttachmentPadding[AttachmentAlignment]{};
+	const auto SerializeAttachmentPadding = [&AttachmentPadding](FArchive& Ar)
+	{
+		const int64 Offset = Ar.Tell();
+		if (const int64 Size = Align(Offset, AttachmentAlignment) - Offset)
+		{
+			Ar.Serialize(AttachmentPadding, Size);
+		}
+	};
+	TUniquePtr<FArchive> PackageAr(IFileManager::Get().CreateFileWriter(*WriteToString<256>(PackagePath), FILEWRITE_NoFail));
+	PackageAr->Serialize(Header, HeaderSize);
+	WorkerObject.CopyTo(*PackageAr);
+	for (int32 AttachmentIndex = 0, AttachmentCount = Attachments.Num(); AttachmentIndex < AttachmentCount; ++AttachmentIndex)
+	{
+		const FCompressedBuffer& Attachment = Attachments[AttachmentIndex];
+		const FBuildWorkerPackageAttachment& AttachmentHeader = Header->Attachments[AttachmentIndex];
+		SerializeAttachmentPadding(*PackageAr);
+		check((uint64)PackageAr->Tell() == Header->AttachmentOffset + AttachmentHeader.CompressedOffset);
+		Attachment.Save(*PackageAr);
+	}
+
+	return Worker;
+}
+
+} // UE::DerivedData::Private
+
+///////////////////////////////////////////////////////////////////////////////////////////////////////////////////////
+
+namespace UE::DerivedData
+{
+
+FBuildWorkerBuilder::FBuildWorkerBuilder()
+	: Builder(MakeUnique<Private::FBuildWorkerBuilderInternal>())
+{
+}
+
+} // UE::DerivedData
diff --git a/Engine/Source/Developer/DerivedDataCache/Private/DerivedDataBuildWorkerRegistry.cpp b/Engine/Source/Developer/DerivedDataCache/Private/DerivedDataBuildWorkerRegistry.cpp
--- a/Engine/Source/Developer/DerivedDataCache/Private/DerivedDataBuildWorkerRegistry.cpp
+++ b/Engine/Source/Developer/DerivedDataCache/Private/DerivedDataBuildWorkerRegistry.cpp
@@ -8,177 +8,11 @@
 #include "Async/UniqueLock.h"
 #include "Containers/SharedString.h"
 #include "Containers/StringView.h"
-#include "DerivedDataBuildKey.h"
-#include "DerivedDataBuildPrivate.h"
 #include "DerivedDataBuildWorker.h"
 #include "Features/IModularFeatures.h"
-#include "Serialization/CompactBinaryWriter.h"
+#include "Misc/Guid.h"
 #include "Templates/Function.h"
 
-namespace UE::DerivedData
-{
-
-FBuildWorker::FBuildWorker(const FBuildWorkerKey& InKey, FCbObject&& InWorker, TUniquePtr<IBuildWorkerResolver> InResolver)
-	: Worker(MoveTemp(InWorker))
-	, Key(InKey)
-	, Resolver(MoveTemp(InResolver))
-{
-}
-
-FBuildWorkerKey FBuildWorker::GetKey() const
-{
-	return Key;
-}
-
-FUtf8StringView FBuildWorker::GetName() const
-{
-	return Worker.FindView(ANSITEXTVIEW("Name")).AsString();
-}
-
-FUtf8StringView FBuildWorker::GetPath() const
-{
-	return Worker.FindView(ANSITEXTVIEW("Path")).AsString();
-}
-
-FUtf8StringView FBuildWorker::GetHostPlatform() const
-{
-	return Worker.FindView(ANSITEXTVIEW("HostPlatform")).AsString();
-}
-
-FGuid FBuildWorker::GetBuildSystemVersion() const
-{
-	return Worker.FindView(ANSITEXTVIEW("BuildSystemVersion")).AsUuid();
-}
-
-void FBuildWorker::IterateFunctions(TFunctionRef<void (FUtf8StringView Name, const FGuid& Version)> Visitor) const
-{
-	for (FCbFieldView It : Worker.FindView(ANSITEXTVIEW("Functions")))
-	{
-		Visitor(It[ANSITEXTVIEW("Name")].AsString(), It[ANSITEXTVIEW("Version")].AsUuid());
-	}
-}
-
-void FBuildWorker::IterateFiles(TFunctionRef<void (FUtf8StringView Path, const FIoHash& RawHash, uint64 RawSize)> Visitor) const
-{
-	for (FCbFieldView It : Worker.FindView(ANSITEXTVIEW("Files")))
-	{
-		Visitor(It[ANSITEXTVIEW("Path")].AsString(), It[ANSITEXTVIEW("RawHash")].AsHash(), It[ANSITEXTVIEW("RawSize")].AsUInt64());
-	}
-}
-
-void FBuildWorker::IterateExecutables(TFunctionRef<void (FUtf8StringView Path, const FIoHash& RawHash, uint64 RawSize)> Visitor) const
-{
-	for (FCbFieldView It : Worker.FindView(ANSITEXTVIEW("Executables")))
-	{
-		Visitor(It[ANSITEXTVIEW("Path")].AsString(), It[ANSITEXTVIEW("RawHash")].AsHash(), It[ANSITEXTVIEW("RawSize")].AsUInt64());
-	}
-}
-
-void FBuildWorker::IterateEnvironment(TFunctionRef<void (FUtf8StringView Name, FUtf8StringView Value)> Visitor) const
-{
-	for (FCbFieldView It : Worker.FindView(ANSITEXTVIEW("Environment")))
-	{
-		Visitor(It[ANSITEXTVIEW("Name")].AsString(), It[ANSITEXTVIEW("Value")].AsString());
-	}
-}
-
-void FBuildWorker::Resolve(TConstArrayView<FIoHash> RawHashes, IRequestOwner& Owner, FOnBuildWorkerResolved&& OnResolved) const
-{
-	Resolver->Resolve(RawHashes, Owner, MoveTemp(OnResolved));
-}
-
-TOptional<FBuildWorker> FBuildWorker::Load(const FBuildWorkerKey& Key, FCbObject&& Worker, TUniquePtr<IBuildWorkerResolver> Resolver)
-{
-	if (Key.Hash == Worker.GetHash())
-	{
-		return FBuildWorker(Key, MoveTemp(Worker), MoveTemp(Resolver));
-	}
-	return {};
-}
-
-FBuildWorker FBuildWorkerBuilder::Build(TUniquePtr<IBuildWorkerResolver> Resolver)
-{
-	Functions.Sort();
-	Files.Sort();
-	Executables.Sort();
-	Environment.Sort();
-
-	TCbWriter<4096> Writer;
-	Writer.BeginObject();
-
-	Writer.AddString(ANSITEXTVIEW("Name"), WorkerName);
-	Writer.AddString(ANSITEXTVIEW("Path"), WorkerPath);
-	Writer.AddString(ANSITEXTVIEW("HostPlatform"), HostPlatform);
-	Writer.AddUuid(ANSITEXTVIEW("BuildSystemVersion"), BuildSystemVersion);
-
-	if (!Functions.IsEmpty())
-	{
-		Writer.BeginArray(ANSITEXTVIEW("Functions"));
-		for (const TTuple<FUtf8String, FGuid>& Function : Functions)
-		{
-			const FUtf8String& Name = Function.Get<FUtf8String>();
-			const FGuid& Version = Function.Get<FGuid>();
-			UE_CLOGF(!Version.IsValid(), LogDerivedDataBuild, Error,
-				"Version of zero is not allowed in build function with the name %s in build worker '%s'.",
-				*WriteToUtf8String<32>(Name), *WorkerName);
-			Writer.BeginObject();
-			Writer.AddString(ANSITEXTVIEW("Name"), Name);
-			Writer.AddUuid(ANSITEXTVIEW("Version"), Version);
-			Writer.EndObject();
-		}
-		Writer.EndArray();
-	}
-
-	if (!Files.IsEmpty())
-	{
-		Writer.BeginArray(ANSITEXTVIEW("Files"));
-		for (const TTuple<FUtf8String, FIoHash, uint64>& File : Files)
-		{
-			Writer.BeginObject();
-			Writer.AddString(ANSITEXTVIEW("Path"), File.Get<FUtf8String>());
-			Writer.AddBinaryAttachment(ANSITEXTVIEW("RawHash"), File.Get<FIoHash>());
-			Writer.AddInteger(ANSITEXTVIEW("RawSize"), File.Get<uint64>());
-			Writer.EndObject();
-		}
-		Writer.EndArray();
-	}
-
-	if (!Executables.IsEmpty())
-	{
-		Writer.BeginArray(ANSITEXTVIEW("Executables"));
-		for (const TTuple<FUtf8String, FIoHash, uint64>& Executable : Executables)
-		{
-			Writer.BeginObject();
-			Writer.AddString(ANSITEXTVIEW("Path"), Executable.Get<FUtf8String>());
-			Writer.AddBinaryAttachment(ANSITEXTVIEW("RawHash"), Executable.Get<FIoHash>());
-			Writer.AddInteger(ANSITEXTVIEW("RawSize"), Executable.Get<uint64>());
-			Writer.EndObject();
-		}
-		Writer.EndArray();
-	}
-
-	if (!Environment.IsEmpty())
-	{
-		Writer.BeginArray(ANSITEXTVIEW("Environment"));
-		for (const TTuple<FUtf8String, FUtf8String>& Var : Environment)
-		{
-			Writer.BeginObject();
-			Writer.AddString(ANSITEXTVIEW("Name"), Var.Get<0>());
-			Writer.AddString(ANSITEXTVIEW("Value"), Var.Get<1>());
-			Writer.EndObject();
-		}
-		Writer.EndArray();
-	}
-
-	Writer.EndObject();
-	FCbObject Worker = Writer.Save().AsObject();
-	return FBuildWorker({Worker.GetHash()}, MoveTemp(Worker), MoveTemp(Resolver));
-}
-
-} // UE::DerivedData
-
-///////////////////////////////////////////////////////////////////////////////////////////////////////////////////////
-
 namespace UE::DerivedData::Private
 {
 
@@ -188,22 +22,28 @@ public:
 	FBuildWorkerRegistry();
 	~FBuildWorkerRegistry();
 
+	void RegisterWorker(FBuildWorker&& Worker, TNotNull<IBuildWorkerFactory*> Factory) final;
+
 	FBuildWorker* FindWorker(
 		const FUtf8SharedString& Function,
 		const FGuid& FunctionVersion,
 		const FGuid& BuildSystemVersion,
 		IBuildWorkerExecutor*& OutWorkerExecutor) const final;
 
+	FBuildWorkerBuilder CreateWorker() const final { return {}; }
+
+	TOptional<FBuildWorker> LoadWorker(FStringView PackagePath) const final { return FBuildWorker::Load(PackagePath); }
+
 private:
 	void OnModularFeatureRegistered(const FName& Type, IModularFeature* ModularFeature);
 	void OnModularFeatureUnregistered(const FName& Type, IModularFeature* ModularFeature);
 
-	void AddWorker(IBuildWorkerFactory* Factory);
-	void RemoveWorker(IBuildWorkerFactory* Factory);
+	void AddWorkers(IBuildWorkerFactory* Factory);
+	void RemoveWorkers(IBuildWorkerFactory* Factory);
 
 private:
 	IBuildWorkerExecutor* Executor = nullptr;
-	TMap<IBuildWorkerFactory*, TUniquePtr<FBuildWorker>> Workers;
+	TMap<IBuildWorkerFactory*, TArray<TUniquePtr<FBuildWorker>>> Workers;
 	TMultiMap<TTuple<FUtf8SharedString, FGuid>, FBuildWorker*> Functions;
 	mutable FSharedMutex Mutex;
 };
@@ -217,7 +57,7 @@ FBuildWorkerRegistry::FBuildWorkerRegistry()
 	}
 	for (IBuildWorkerFactory* Worker : ModularFeatures.GetModularFeatureImplementations<IBuildWorkerFactory>(IBuildWorkerFactory::FeatureName))
 	{
-		AddWorker(Worker);
+		AddWorkers(Worker);
 	}
 	ModularFeatures.OnModularFeatureRegistered().AddRaw(this, &FBuildWorkerRegistry::OnModularFeatureRegistered);
 	ModularFeatures.OnModularFeatureUnregistered().AddRaw(this, &FBuildWorkerRegistry::OnModularFeatureUnregistered);
@@ -239,7 +79,7 @@ void FBuildWorkerRegistry::OnModularFeatureRegistered(const FName& Type, IModula
 	}
 	else if (Type == IBuildWorkerFactory::FeatureName)
 	{
-		AddWorker(static_cast<IBuildWorkerFactory*>(ModularFeature));
+		AddWorkers(static_cast<IBuildWorkerFactory*>(ModularFeature));
 	}
 }
 
@@ -258,34 +98,42 @@ void FBuildWorkerRegistry::OnModularFeatureUnregistered(const FName& Type, IModu
 	}
 	else if (Type == IBuildWorkerFactory::FeatureName)
 	{
-		RemoveWorker(static_cast<IBuildWorkerFactory*>(ModularFeature));
+		RemoveWorkers(static_cast<IBuildWorkerFactory*>(ModularFeature));
 	}
 }
 
-void FBuildWorkerRegistry::AddWorker(IBuildWorkerFactory* Factory)
+void FBuildWorkerRegistry::AddWorkers(IBuildWorkerFactory* Factory)
+{
+	Factory->CreateWorkers(*this);
+}
+
+void FBuildWorkerRegistry::RemoveWorkers(IBuildWorkerFactory* Factory)
 {
-	FBuildWorkerBuilder Builder;
-	if (TOptional<FBuildWorker> OptionalWorker = Factory->Build(Builder))
+	Factory->AbortCreateWorkers();
+
+	TUniqueLock Lock(Mutex);
+	if (TArray<TUniquePtr<FBuildWorker>>* WorkerPtrs = Workers.Find(Factory))
 	{
-		TUniquePtr<FBuildWorker> Worker = MakeUnique<FBuildWorker>(MoveTemp(OptionalWorker.GetValue()));
-		TUniqueLock Lock(Mutex);
-		Worker->IterateFunctions([this, Worker = Worker.Get()](FUtf8StringView Name, const FGuid& Version)
+		for (TUniquePtr<FBuildWorker>& Worker : *WorkerPtrs)
 		{
-			Functions.Emplace(MakeTuple(FUtf8SharedString(Name), Version), Worker);
-		});
-		Workers.Emplace(Factory, MoveTemp(Worker));
+			Worker->IterateFunctions([this, Worker = Worker.Get()](FUtf8StringView Name, const FGuid& Version)
+			{
+				Functions.Remove(MakeTuple(FUtf8SharedString(Name), Version), Worker);
+			});
+		}
+		Workers.Remove(Factory);
 	}
 }
 
-void FBuildWorkerRegistry::RemoveWorker(IBuildWorkerFactory* Factory)
+void FBuildWorkerRegistry::RegisterWorker(FBuildWorker&& InWorker, TNotNull<IBuildWorkerFactory*> Factory)
 {
+	TUniquePtr<FBuildWorker> Worker = MakeUnique<FBuildWorker>(MoveTemp(InWorker));
 	TUniqueLock Lock(Mutex);
-	TUniquePtr<FBuildWorker>& Worker = Workers.FindChecked(Factory);
 	Worker->IterateFunctions([this, Worker = Worker.Get()](FUtf8StringView Name, const FGuid& Version)
 	{
-		Functions.Remove(MakeTuple(FUtf8SharedString(Name), Version), Worker);
+		Functions.Emplace(MakeTuple(FUtf8SharedString(Name), Version), Worker);
 	});
-	Workers.Remove(Factory);
+	Workers.FindOrAdd(Factory).Emplace(MoveTemp(Worker));
 }
 
 FBuildWorker* FBuildWorkerRegistry::FindWorker(
diff --git a/Engine/Source/Developer/DerivedDataCache/Public/DerivedDataBuildWorker.h b/Engine/Source/Developer/DerivedDataCache/Public/DerivedDataBuildWorker.h
--- a/Engine/Source/Developer/DerivedDataCache/Public/DerivedDataBuildWorker.h
+++ b/Engine/Source/Developer/DerivedDataCache/Public/DerivedDataBuildWorker.h
@@ -3,38 +3,59 @@
 #pragma once
 
 #include "CoreTypes.h"
-#include "Containers/Array.h"
 #include "Containers/ArrayView.h"
 #include "Containers/StringFwd.h"
-#include "Containers/Utf8String.h"
 #include "DerivedDataBuildKey.h"
 #include "DerivedDataRequestTypes.h"
 #include "Features/IModularFeature.h"
 #include "IO/IoHash.h"
-#include "Misc/Guid.h"
 #include "Misc/NotNull.h"
 #include "Misc/OptionalFwd.h"
 #include "Serialization/CompactBinary.h"
-#include "Templates/Tuple.h"
+#include "Templates/UniquePtr.h"
 #include "UObject/NameTypes.h"
 
 #define UE_API DERIVEDDATACACHE_API
 
 class FCompressedBuffer;
+struct FGuid;
 
 template <typename FuncType> class TFunctionRef;
 template <typename FuncType> class TUniqueFunction;
 
 namespace UE::DerivedData { class FBuildAction; }
 namespace UE::DerivedData { class FBuildPolicy; }
-namespace UE::DerivedData { class FBuildWorkerBuilder; }
+namespace UE::DerivedData { class FBuildWorker; }
 namespace UE::DerivedData { class FOptionalBuildInputs; }
 namespace UE::DerivedData { class FOptionalBuildOutput; }
 namespace UE::DerivedData { class IBuild; }
+namespace UE::DerivedData { class IBuildWorkerRegistry; }
 namespace UE::DerivedData { class IBuildWorkerResolver; }
 namespace UE::DerivedData { class IRequestOwner; }
 namespace UE::DerivedData { struct FBuildWorkerActionCompleteParams; }
 namespace UE::DerivedData { struct FBuildWorkerResolvedParams; }
+namespace UE::DerivedData::Private { class FBuildWorkerBuilderInternal; }
+
+namespace UE::DerivedData::Private
+{
+
+class IBuildWorkerBuilderInternal
+{
+public:
+	virtual ~IBuildWorkerBuilderInternal() = default;
+	virtual void SetName(FUtf8StringView Name) = 0;
+	virtual void SetPath(FUtf8StringView Path) = 0;
+	virtual void SetHostPlatform(FUtf8StringView Name) = 0;
+	virtual void SetBuildSystemVersion(const FGuid& Version) = 0;
+	virtual void AddFunction(FUtf8StringView Name, const FGuid& Version) = 0;
+	virtual void AddDirectory(FUtf8StringView Path) = 0;
+	virtual void AddFile(FUtf8StringView Path, const FCompressedBuffer& Data) = 0;
+	virtual void AddExecutable(FUtf8StringView Path, const FCompressedBuffer& Data) = 0;
+	virtual void SetEnvironment(FUtf8StringView Name, FUtf8StringView Value) = 0;
+	virtual FBuildWorker Build(FStringView PackagePath) = 0;
+};
+
+} // UE::DerivedData::Private
 
 namespace UE::DerivedData
 {
@@ -46,10 +67,31 @@ using FOnBuildWorkerResolved = TUniqueFunction<void (FBuildWorkerResolvedParams&
 struct FBuildWorkerKey final
 {
 	FIoHash Hash;
+
+	friend inline bool operator==(const FBuildWorkerKey& A, const FBuildWorkerKey& B)
+	{
+		return A.Hash == B.Hash;
+	}
+
+	friend inline bool operator<(const FBuildWorkerKey& A, const FBuildWorkerKey& B)
+	{
+		return A.Hash < B.Hash;
+	}
+
+	friend inline uint32 GetTypeHash(const FBuildWorkerKey& Key)
+	{
+		return GetTypeHashHelper(Key.Hash);
+	}
+
+	template <typename CharType>
+	friend inline TStringBuilderBase<CharType>& operator<<(TStringBuilderBase<CharType>& Builder, const FBuildWorkerKey& Key)
+	{
+		return Builder << ANSITEXTVIEW("BuildWorker/") << Key.Hash;
+	}
 };
 
 /**
- * A build worker is an immutable reference to a build worker program and its dependencies.
+ * A build worker is a reference to an immutable build worker program and its dependencies.
  *
  * The purpose of a build worker is to capture which build functions a program contains and refer
  * to everything required to execute the program, including files and environment.
@@ -62,131 +104,141 @@ struct FBuildWorkerKey final
 class FBuildWorker final
 {
 public:
-	UE_API FBuildWorkerKey GetKey() const;
-	UE_API FUtf8StringView GetName() const;
-	UE_API FUtf8StringView GetPath() const;
-	UE_API FUtf8StringView GetHostPlatform() const;
-	UE_API FGuid GetBuildSystemVersion() const;
+	[[nodiscard]] inline FBuildWorkerKey GetKey() const
+	{
+		return Key;
+	}
+
+	[[nodiscard]] UE_API FUtf8StringView GetName() const;
+	[[nodiscard]] UE_API FUtf8StringView GetPath() const;
+	[[nodiscard]] UE_API FUtf8StringView GetHostPlatform() const;
+	[[nodiscard]] UE_API FGuid GetBuildSystemVersion() const;
+
 	UE_API void IterateFunctions(TFunctionRef<void (FUtf8StringView Name, const FGuid& Version)> Visitor) const;
+	UE_API void IterateDirectories(TFunctionRef<void (FUtf8StringView Path)> Visitor) const;
 	UE_API void IterateFiles(TFunctionRef<void (FUtf8StringView Path, const FIoHash& RawHash, uint64 RawSize)> Visitor) const;
 	UE_API void IterateExecutables(TFunctionRef<void (FUtf8StringView Path, const FIoHash& RawHash, uint64 RawSize)> Visitor) const;
 	UE_API void IterateEnvironment(TFunctionRef<void (FUtf8StringView Name, FUtf8StringView Value)> Visitor) const;
 
-	/** Resolve files and executables referenced by this worker, by raw hash. */
+	/** Resolves files and executables referenced by this build worker, by raw hash. */
 	UE_API void Resolve(TConstArrayView<FIoHash> RawHashes, IRequestOwner& Owner, FOnBuildWorkerResolved&& OnResolved) const;
 
-	/** Saves the build worker to a compact binary object. */
-	inline const FCbObject& Save(FCbWriter& Writer) const
+	/** Extracts the build worker into a directory, which must be empty. */
+	UE_API bool Extract(FStringView OutputPath) const;
+
+	/** Saves the build worker to a compact binary object, excluding file data. */
+	[[nodiscard]] inline const FCbObject& Save() const
 	{
-		return Worker;
+		return Object;
 	}
 
+	/** Path to the package that this worker was loaded from. Empty when not loaded from a package. */
+	[[nodiscard]] UE_API FStringView GetPackagePath() const;
+
+	/** Loads a build worker from a package created by FBuildWorkerBuilder::Build. */
+	[[nodiscard]] UE_API static TOptional<FBuildWorker> Load(FStringView PackagePath);
+
 	/**
 	 * Loads a build worker from a compact binary object.
 	 *
-	 * @param Key        The key of the worker being loaded. The key is validated during load.
-	 * @param Worker     An object saved from a build worker. Holds a reference and is cloned if not owned.
+	 * @param Object     An object saved from a build worker. Holds a reference and is cloned if not owned.
 	 * @param Resolver   A resolver used to find file data referenced by the build worker.
-	 * @return A valid build worker, or unset on error.
 	 */
-	UE_API static TOptional<FBuildWorker> Load(const FBuildWorkerKey& Key, FCbObject&& Worker, TUniquePtr<IBuildWorkerResolver> Resolver);
-
-	static constexpr bool bHasIntrusiveUnsetOptionalState = true;
-
-	inline explicit FBuildWorker(FIntrusiveUnsetOptionalState)
-	{
-	}
+	[[nodiscard]] UE_API static FBuildWorker Load(FCbObject Object, TNotNull<TUniquePtr<IBuildWorkerResolver>> Resolver);
 
-	inline bool operator==(FIntrusiveUnsetOptionalState) const
-	{
-		return !Resolver;
-	}
+	constexpr static bool bHasIntrusiveUnsetOptionalState = true;
+	using IntrusiveUnsetOptionalStateType = FBuildWorker;
+	[[nodiscard]] inline explicit FBuildWorker(FIntrusiveUnsetOptionalState) {}
+	[[nodiscard]] inline bool operator==(FIntrusiveUnsetOptionalState) const { return !Resolver; }
 
 private:
-	explicit FBuildWorker(const FBuildWorkerKey& Key, FCbObject&& Worker, TUniquePtr<IBuildWorkerResolver> Resolver);
+	FBuildWorker() = default;
 
-	FCbObject Worker;
-	FBuildWorkerKey Key;
+	FCbObject Object;
 	TUniquePtr<IBuildWorkerResolver> Resolver;
-
-	friend FBuildWorkerBuilder;
+	FBuildWorkerKey Key;
 };
 
 /** A builder to construct a build worker. */
-class FBuildWorkerBuilder
+class FBuildWorkerBuilder final
 {
 public:
-	// Destructor is virtual to allow an include-only dependency on this module.
-	virtual ~FBuildWorkerBuilder() = default;
+	[[nodiscard]] UE_API FBuildWorkerBuilder();
 
 	inline void SetName(FUtf8StringView Name)
 	{
-		WorkerName = Name;
+		Builder->SetName(Name);
 	}
 
 	inline void SetPath(FUtf8StringView Path)
 	{
-		WorkerPath = Path;
+		Builder->SetPath(Path);
 	}
 
 	inline void SetHostPlatform(FUtf8StringView Name)
 	{
-		HostPlatform = Name;
+		Builder->SetHostPlatform(Name);
 	}
 
 	inline void SetBuildSystemVersion(const FGuid& Version)
 	{
-		BuildSystemVersion = Version;
+		Builder->SetBuildSystemVersion(Version);
 	}
 
 	inline void AddFunction(FUtf8StringView Name, const FGuid& Version)
 	{
-		Functions.Emplace(Name, Version);
+		Builder->AddFunction(Name, Version);
 	}
 
-	inline void AddFile(FUtf8StringView Path, const FIoHash& RawHash, uint64 RawSize)
+	inline void AddDirectory(FUtf8StringView Path)
 	{
-		Files.Emplace(Path, RawHash, RawSize);
+		Builder->AddDirectory(Path);
 	}
 
-	inline void AddSymbolFile(FUtf8StringView Path, const FIoHash& RawHash, uint64 RawSize)
+	inline void AddFile(FUtf8StringView Path, const FCompressedBuffer& Data)
 	{
-		SymbolFiles.Emplace(Path, RawHash, RawSize);
+		Builder->AddFile(Path, Data);
 	}
 
-	inline void AddExecutable(FUtf8StringView Path, const FIoHash& RawHash, uint64 RawSize)
+	inline void AddExecutable(FUtf8StringView Path, const FCompressedBuffer& Data)
 	{
-		Executables.Emplace(Path, RawHash, RawSize);
+		Builder->AddExecutable(Path, Data);
 	}
 
 	inline void SetEnvironment(FUtf8StringView Name, FUtf8StringView Value)
 	{
-		Environment.Emplace(Name, Value);
+		Builder->SetEnvironment(Name, Value);
 	}
 
-	// Build is virtual to allow an include-only dependency on this module.
-	UE_API virtual FBuildWorker Build(TUniquePtr<IBuildWorkerResolver> Resolver);
+	[[nodiscard]] inline FBuildWorker Build(FStringView PackagePath)
+	{
+		return Builder->Build(PackagePath);
+	}
 
 private:
-	FUtf8String WorkerName;
-	FUtf8String WorkerPath;
-	FUtf8String HostPlatform;
-	FGuid BuildSystemVersion;
-	TArray<TTuple<FUtf8String, FGuid>> Functions;
-	TArray<TTuple<FUtf8String, FIoHash, uint64>> Files;
-	TArray<TTuple<FUtf8String, FIoHash, uint64>> SymbolFiles;
-	TArray<TTuple<FUtf8String, FIoHash, uint64>> Executables;
-	TArray<TTuple<FUtf8String, FUtf8String>> Environment;
+	TUniquePtr<Private::IBuildWorkerBuilderInternal> Builder;
 };
 
+/** Interface to resolve file data for a build worker. */
 class IBuildWorkerResolver
 {
 public:
 	virtual ~IBuildWorkerResolver() = default;
 
+	/**
+	 * Resolve the compressed data for the raw hashes.
+	 *
+	 * @param RawHashes   The raw hashes of the data to resolve.
+	 * @param Owner   Owner provides priority and manages requests when resolving is asynchronous.
+	 * @param OnResolved   Called exactly once when data is resolved. May return partial results if Status is Error.
+	 */
 	virtual void Resolve(TConstArrayView<FIoHash> RawHashes, IRequestOwner& Owner, FOnBuildWorkerResolved&& OnResolved) = 0;
+
+	/** Returns the path to the package that this is resolving from, or empty if not backed by a package. */
+	[[nodiscard]] virtual FStringView GetPackagePath() const = 0;
 };
 
+/** Interface to create and register build workers with the build worker registry. */
 class IBuildWorkerFactory : public IModularFeature
 {
 public:
@@ -194,10 +246,22 @@ public:
 
 	virtual ~IBuildWorkerFactory() = default;
 
-	// Builder is provided to allow an include-only dependency on this module.
-	virtual TOptional<FBuildWorker> Build(FBuildWorkerBuilder& Builder) = 0;
+	/**
+	 * Called by the build system when it is created, or when the factory is registered later.
+	 *
+	 * Async worker creation is allowed and must be finished before AbortCreateWorkers() returns.
+	 */
+	virtual void CreateWorkers(IBuildWorkerRegistry& WorkerRegistry) = 0;
+
+	/**
+	 * Called by the build system when it is about to be destroyed.
+	 *
+	 * Any async worker creation must be completed or aborted before returning from this function.
+	 */
+	virtual void AbortCreateWorkers() = 0;
 };
 
+/** Interface to execute a build action using a build worker. */
 class IBuildWorkerExecutor : public IModularFeature
 {
 public:
diff --git a/Engine/Source/Developer/DerivedDataCache/Public/DerivedDataBuildWorkerRegistry.h b/Engine/Source/Developer/DerivedDataCache/Public/DerivedDataBuildWorkerRegistry.h
--- a/Engine/Source/Developer/DerivedDataCache/Public/DerivedDataBuildWorkerRegistry.h
+++ b/Engine/Source/Developer/DerivedDataCache/Public/DerivedDataBuildWorkerRegistry.h
@@ -2,13 +2,18 @@
 
 #pragma once
 
-#include "Containers/StringFwd.h"
 #include "CoreTypes.h"
+#include "Containers/StringFwd.h"
+#include "Misc/NotNull.h"
+#include "Misc/OptionalFwd.h"
 
 struct FGuid;
 
 namespace UE::DerivedData { class FBuildWorker; }
+namespace UE::DerivedData { class FBuildWorkerBuilder; }
+namespace UE::DerivedData { class FBuildWorkerBuilder; }
 namespace UE::DerivedData { class IBuildWorkerExecutor; }
+namespace UE::DerivedData { class IBuildWorkerFactory; }
 
 namespace UE::DerivedData
 {
@@ -21,6 +26,9 @@ class IBuildWorkerRegistry
 public:
 	virtual ~IBuildWorkerRegistry() = default;
 
+	/** Registers a build worker and takes ownership of it. */
+	virtual void RegisterWorker(FBuildWorker&& Worker, TNotNull<IBuildWorkerFactory*> Factory) = 0;
+
 	/**
 	 * Finds a build worker that can execute the function at the version.
 	 *
@@ -35,6 +43,12 @@ public:
 		const FGuid& FunctionVersion,
 		const FGuid& BuildSystemVersion,
 		IBuildWorkerExecutor*& OutWorkerExecutor) const = 0;
+
+	/** Creates a builder to create a build worker. Load with LoadWorker. */
+	virtual FBuildWorkerBuilder CreateWorker() const = 0;
+
+	/** Loads a build worker as created through CreateWorker. */
+	virtual TOptional<FBuildWorker> LoadWorker(FStringView PackagePath) const = 0;
 };
 
 } // UE::DerivedData
diff --git a/Engine/Source/Developer/DesktopPlatform/Private/TargetReceiptBuildWorker.cpp b/Engine/Source/Developer/DesktopPlatform/Private/TargetReceiptBuildWorker.cpp
--- a/Engine/Source/Developer/DesktopPlatform/Private/TargetReceiptBuildWorker.cpp
+++ b/Engine/Source/Developer/DesktopPlatform/Private/TargetReceiptBuildWorker.cpp
@@ -2,20 +2,25 @@
 
 #include "TargetReceiptBuildWorker.h"
 
+#include "Algo/Accumulate.h"
+#include "Algo/MaxElement.h"
 #include "Algo/Sort.h"
 #include "Algo/Unique.h"
-#include "Async/TaskGraphInterfaces.h"
+#include "Async/ParallelFor.h"
 #include "Compression/CompressedBuffer.h"
 #include "Compression/OodleDataCompression.h"
 #include "Containers/Array.h"
 #include "Containers/Map.h"
+#include "Containers/SharedString.h"
 #include "Containers/StringView.h"
 #include "Containers/UnrealString.h"
 #include "Containers/Utf8String.h"
 #include "DerivedDataBuildWorker.h"
+#include "DerivedDataBuildWorkerRegistry.h"
 #include "Features/IModularFeatures.h"
 #include "HAL/FileManager.h"
 #include "HAL/Platform.h"
+#include "HAL/PlatformMath.h"
 #include "HAL/PlatformProcess.h"
 #include "HAL/PlatformTime.h"
 #include "IO/IoHash.h"
@@ -26,561 +31,679 @@
 #include "Misc/StringBuilder.h"
 #include "Serialization/CompactBinary.h"
 #include "Serialization/CompactBinarySerialization.h"
+#include "String/ParseTokens.h"
 #include "TargetReceipt.h"
+#include "Tasks/Task.h"
+#include "Templates/Projection.h"
 #include "Templates/Tuple.h"
 
 #include <atomic>
 
 DEFINE_LOG_CATEGORY_STATIC(LogTargetReceiptBuildWorker, Log, All);
 
-class FTargetReceiptBuildWorkerLoader
+class FTargetReceiptBuildWorkerBuilder
 {
 	// This needs to match the version exported in DerivedDataBuildWorker.Build.cs.
-	inline static const FGuid WorkerReceiptVersion{TEXT("dab5352e-a5a7-4793-a7a3-1d4acad6aff2")};
+	inline static const FGuid CurrentWorkerReceiptVersion{TEXT("dab5352e-a5a7-4793-a7a3-1d4acad6aff2")};
 
 public:
-	bool Load(FStringView TargetReceiptFilePath);
-	void AbortLoad();
+	[[nodiscard]] inline explicit FTargetReceiptBuildWorkerBuilder(FStringView InTargetReceiptPath, const FTargetReceiptBuildWorkerParams& InParams)
+		: TargetReceiptPath(InTargetReceiptPath)
+		, Params(InParams)
+		, RelativeProjectDir([this]
+			{
+				FString Path = FPaths::ConvertRelativePathToFull(ProjectDir);
+				FString RootDir = FPaths::ConvertRelativePathToFull(FPaths::RootDir());
+				if (!FPaths::MakePathRelativeTo(Path, *RootDir))
+				{
+					Path = TEXTVIEW("Project");
+				}
+				return Path;
+			}())
+	{
+	}
 
-	UE::DerivedData::FBuildWorker Build(UE::DerivedData::FBuildWorkerBuilder& Builder) const;
+	[[nodiscard]] TOptional<UE::DerivedData::FBuildWorker> Build(UE::DerivedData::IBuildWorkerRegistry& Registry);
 
 private:
-	bool TryAddExecutablePath(FStringView WidePath);
+	struct FFileMetaData
+	{
+		UE::FSharedString LocalPath;
+		FString RemotePath;
+		FDateTime ModTime;
+		FSharedBuffer RawData;
+		FCompressedBuffer Data;
+		bool bExecutable = false;
+	};
 
-	/**
-	 * Encapsulates a path used by the worker that has both a local and remote representation.
-	 */
-	struct FWorkerPath
-	{
-		explicit FWorkerPath(FUtf8StringView InLocalPathRoot, FUtf8StringView InLocalPathSuffix, FUtf8StringView InRemotePathRoot = {}, FUtf8StringView InRemotePathSuffix = {});
+	[[nodiscard]] bool TryFindTargetReceipt(FStringBuilderBase& OutActiveTargetReceiptPath, FStringBuilderBase& OutActiveWorkerPackagePath);
 
-		FUtf8String LocalPath;
-		FUtf8String RemotePath;
-	};
+	[[nodiscard]] FFileMetaData* TryAddFile(FStringView Path);
+
+	[[nodiscard]] bool TryLoadWorkerVersions(FStringView ActiveTargetReceiptPath, FStringView ActiveWorkerPackagePath);
 
-	FUtf8String Name;
-	FUtf8String Platform;
+private:
+	const FStringView TargetReceiptPath;
+	const FTargetReceiptBuildWorkerParams Params;
+
+	FString Name;
+	FString HostPlatform;
 	FGuid BuildSystemVersion;
-	TArray<TTuple<FUtf8String, FGuid>> FunctionVersions;
-	using FWorkerPathMeta = TTuple<FIoHash, uint64>;
-	TArray<FWorkerPathMeta> ExecutableMeta;
-	TArray<FWorkerPath> ExecutablePaths;
-	TMap<FUtf8String, FUtf8String> EnvironmentVariables;
-	std::atomic<bool> bAbortRequested = false;
-};
+	TArray<TTuple<FString, FGuid>> Functions;
+	TMap<UE::FSharedString, FFileMetaData> LocalPathToMetaData;
+	TMap<FString, FString> Environment;
 
-FTargetReceiptBuildWorkerLoader::FWorkerPath::FWorkerPath(FUtf8StringView InLocalPathRoot, FUtf8StringView InLocalPathSuffix, FUtf8StringView InRemotePathRoot, FUtf8StringView InRemotePathSuffix)
-{
-	TUtf8StringBuilder<128> PathBuilder;
-	FPathViews::Append(PathBuilder, InLocalPathRoot, InLocalPathSuffix);
-	LocalPath = PathBuilder;
-
-	PathBuilder.Reset();
-	FPathViews::Append(PathBuilder,
-		InRemotePathRoot.IsEmpty() ? InLocalPathRoot : InRemotePathRoot,
-		InRemotePathSuffix.IsEmpty() ? InLocalPathSuffix : InRemotePathSuffix);
-	RemotePath = PathBuilder;
-}
+	UE::FSharedString LocalLaunchPath;
+	FDateTime MostRecentFileModTime;
 
-#if WITH_EDITOR
-class FTargetReceiptBuildWorkerFactory final : public UE::DerivedData::IBuildWorkerFactory
+	IFileManager& FileManager = IFileManager::Get();
+	const FString EngineDir = FPaths::EngineDir();
+	const FString ProjectDir = FPaths::ProjectDir();
+	const FString RelativeProjectDir;
+	const FStringView EngineDirPrefix = TEXTVIEW("$(EngineDir)");
+	const FStringView ProjectDirPrefix = TEXTVIEW("$(ProjectDir)");
+};
+
+TOptional<UE::DerivedData::FBuildWorker> FTargetReceiptBuildWorkerBuilder::Build(UE::DerivedData::IBuildWorkerRegistry& Registry)
 {
-public:
-	explicit FTargetReceiptBuildWorkerFactory(FStringView TargetReceiptFilePath);
+	const double StartTime = FPlatformTime::Seconds();
 
-	~FTargetReceiptBuildWorkerFactory()
+	TStringBuilder<256> ActiveTargetReceiptPath;
+	TStringBuilder<256> ActiveWorkerPackagePath(InPlace, Params.WorkerPackagePath);
+	if (!TryFindTargetReceipt(ActiveTargetReceiptPath, ActiveWorkerPackagePath))
 	{
-		Loader.AbortLoad();
-		if (LoadTargetReceiptTask.IsValid())
-		{
-			LoadTargetReceiptTask->Wait();
-		}
-
-		if (bRegistered)
-		{
-			IModularFeatures::Get().UnregisterModularFeature(UE::DerivedData::IBuildWorkerFactory::FeatureName, this);
-		}
+		return {};
 	}
 
-	void LoadAndRegister(FStringView TargetReceiptFilePath)
+	const FDateTime TargetReceiptModTime = FileManager.GetTimeStamp(*WriteToString<256>(ActiveTargetReceiptPath));
+	const FDateTime WorkerPackageModTime = FileManager.GetTimeStamp(*WriteToString<256>(ActiveWorkerPackagePath));
+
+	if (TargetReceiptModTime == FDateTime::MinValue())
 	{
-		if (Loader.Load(TargetReceiptFilePath))
+		if (TOptional<UE::DerivedData::FBuildWorker> Worker = Registry.LoadWorker(ActiveWorkerPackagePath))
 		{
-			check(!bRegistered);
-			bRegistered = true;
-			IModularFeatures::Get().RegisterModularFeature(UE::DerivedData::IBuildWorkerFactory::FeatureName, this);
+			UE_LOGFMT(LogTargetReceiptBuildWorker, Verbose,
+				"Target receipt does not exist at '{ActiveTargetReceipt}'. "
+				"Using existing build worker package at '{Path}' for '{TargetReceipt}'.",
+				ActiveTargetReceiptPath, ActiveWorkerPackagePath, TargetReceiptPath);
+			return Worker;
 		}
+		return {};
 	}
 
-	TOptional<UE::DerivedData::FBuildWorker> Build(UE::DerivedData::FBuildWorkerBuilder& Builder) final
-	{
-		return Loader.Build(Builder);
-	}
-
-private:
-	FGraphEventRef LoadTargetReceiptTask;
-	FTargetReceiptBuildWorkerLoader Loader;
-	bool bRegistered = false;
-};
+	UE_LOGFMT(LogTargetReceiptBuildWorker, Verbose, "Using target receipt '{TargetReceipt}'", ActiveTargetReceiptPath);
 
-class FPopulateWorkerFromTargetReceiptTask
-{
-public:
-	explicit FPopulateWorkerFromTargetReceiptTask(FStringView InTargetReceiptFilePath, FTargetReceiptBuildWorkerFactory* InFactory)
-		: TargetReceiptFilePath(InTargetReceiptFilePath)
-		, Factory(InFactory)
+	if (WorkerPackageModTime > TargetReceiptModTime && !Params.bAlwaysBuild)
 	{
+		if (TOptional<UE::DerivedData::FBuildWorker> Worker = Registry.LoadWorker(ActiveWorkerPackagePath))
+		{
+			UE_LOGFMT(LogTargetReceiptBuildWorker, Verbose,
+				"Target receipt has not been modified since the build worker package was created. "
+				"Using existing build worker package at '{Path}' for '{TargetReceipt}'.",
+				ActiveWorkerPackagePath, TargetReceiptPath);
+			return Worker;
+		}
 	}
 
-	inline TStatId GetStatId() const { RETURN_QUICK_DECLARE_CYCLE_STAT(FPopulateWorkerFromTargetReceiptTask, STATGROUP_TaskGraphTasks); }
-	constexpr ENamedThreads::Type GetDesiredThread() { return ENamedThreads::Type::AnyThread; }
-	constexpr static ESubsequentsMode::Type GetSubsequentsMode() { return ESubsequentsMode::FireAndForget; }
-
-	void DoTask(ENamedThreads::Type CurrentThread, const FGraphEventRef& MyCompletionGraphEvent)
+	FTargetReceipt TargetReceipt;
+	if (!TargetReceipt.Read(FString(ActiveTargetReceiptPath), /*bExpandVariables*/ false))
 	{
-		Factory->LoadAndRegister(TargetReceiptFilePath);
+		UE_LOGFMT(LogTargetReceiptBuildWorker, Warning, "Failed to parse target receipt '{TargetReceipt}'", TargetReceiptPath);
+		return {};
 	}
 
-private:
-	const FString TargetReceiptFilePath;
-	FTargetReceiptBuildWorkerFactory* Factory;
-};
-
-FTargetReceiptBuildWorkerFactory::FTargetReceiptBuildWorkerFactory(FStringView TargetReceiptFilePath)
-{
-	LoadTargetReceiptTask = TGraphTask<FPopulateWorkerFromTargetReceiptTask>::CreateTask().ConstructAndDispatchWhenReady(TargetReceiptFilePath, this);
-}
-#endif // WITH_EDITOR
-
-FTargetReceiptBuildWorker::FTargetReceiptBuildWorker(FStringView TargetReceiptFilePath)
-{
-#if WITH_EDITOR
-	Factory = MakePimpl<FTargetReceiptBuildWorkerFactory>(TargetReceiptFilePath);
-#endif
-}
-
-FTargetReceiptBuildWorker::~FTargetReceiptBuildWorker() = default;
-
-bool FTargetReceiptBuildWorkerLoader::TryAddExecutablePath(FStringView WidePath)
-{
-	TUtf8StringBuilder<256> Path(InPlace, WidePath);
-
-	if (Path.ToView().IsEmpty())
+	if (Params.ShouldAbort && Params.ShouldAbort())
 	{
-		UE_LOGF(LogTargetReceiptBuildWorker, Warning, "Empty executable path encountered while populating worker '%s': '%s'", *Name, *Path);
-		return false;
+		return {};
 	}
 
-	if (!Path.ToView().StartsWith(ANSITEXTVIEW("$(")))
-	{
-		UE_LOGF(LogTargetReceiptBuildWorker, Warning, "Absolute or unprefixed executable path variable encountered while populating worker '%s': '%s'", *Name, *Path);
-		return false;
-	}
+	// Store the name and platform early to allow their use during loading and logging.
+	Name = TargetReceipt.TargetName;
+	HostPlatform = TargetReceipt.Platform;
 
-	int32 EndVariableIndex = Path.ToView().Find(ANSITEXTVIEW(")"));
-	if (EndVariableIndex == INDEX_NONE)
+	if (FFileMetaData* File = TryAddFile(TargetReceipt.Launch))
 	{
-		UE_LOGF(LogTargetReceiptBuildWorker, Warning, "Malformed executable path variable encountered while populating worker '%s': '%s'", *Name, *Path);
-		return false;
+		File->bExecutable = true;
+		LocalLaunchPath = File->LocalPath;
 	}
-
-	FUtf8StringView PrefixVariable = Path.ToView().Mid(2, EndVariableIndex - 2);
-	if (PrefixVariable == ANSITEXTVIEW("EngineDir"))
+	else
 	{
-		ExecutablePaths.Emplace(FUtf8String(FPaths::EngineDir()), Path.ToView().RightChop(EndVariableIndex + 2), ANSITEXTVIEW("Engine"));
+		return {};
 	}
-	else if (PrefixVariable == ANSITEXTVIEW("ProjectDir"))
+
+	for (const FBuildProduct& BuildProduct : TargetReceipt.BuildProducts)
 	{
-		FString ProjectDir = FPaths::ProjectDir();
-		FString RelativeProjectDir = FPaths::ConvertRelativePathToFull(ProjectDir);
-		FString RootDir = FPaths::ConvertRelativePathToFull(FPaths::RootDir());
-		if (FPaths::MakePathRelativeTo(RelativeProjectDir, *RootDir))
+		if (!Params.bIncludeSymbols && BuildProduct.Type == TEXTVIEW("SymbolFile"))
 		{
-			ExecutablePaths.Emplace(FUtf8String(ProjectDir), Path.ToView().RightChop(EndVariableIndex + 2), WriteToUtf8String<64>(RelativeProjectDir));
+			continue;
 		}
-		else
+		if (FFileMetaData* File = TryAddFile(BuildProduct.Path))
 		{
-			ExecutablePaths.Emplace(FUtf8String(ProjectDir), Path.ToView().RightChop(EndVariableIndex + 2), ANSITEXTVIEW("Project"));
+			File->bExecutable = BuildProduct.Type == TEXTVIEW("Executable");
 		}
-	}
-	else
-	{
-		FString TempVariableName(PrefixVariable);
-		FString EnvironmentVariableValue = FPlatformMisc::GetEnvironmentVariable(*TempVariableName);
-		if (EnvironmentVariableValue.IsEmpty())
+		else
 		{
-			UE_LOGF(LogTargetReceiptBuildWorker, Warning, "Unknown executable path variable encountered while populating worker '%s': '%ls'", *Name, *TempVariableName);
-			return false;
+			return {};
 		}
-		ExecutablePaths.Emplace(FUtf8String(EnvironmentVariableValue), Path.ToView().RightChop(EndVariableIndex + 2), PrefixVariable);
-		// By having a dependency that starts with an environment variable, we implicitly create an environment variable to be remoted.
-		// The variable will have the same name as the variable we used, and its value will be a path relative to the platform specific binaries directory with the same name as the variable.
-		EnvironmentVariables.FindOrAdd(FUtf8String(PrefixVariable), FUtf8String(ANSITEXTVIEW("..\\..\\..\\")) + PrefixVariable);
 	}
 
-	return true;
-}
-
-bool FTargetReceiptBuildWorkerLoader::Load(FStringView TargetReceiptFilePath)
-{
-	IFileManager& FileManager = IFileManager::Get();
-	FString FinalTargetReceiptFilePath;
-	FString FinalPublishedVersionFilePath;
-	FString FinalGeneratedVersionFilePath;
-	const FStringView EngineDirPrefix = TEXTVIEW("$(EngineDir)");
-	const FStringView ProjectDirPrefix = TEXTVIEW("$(ProjectDir)");
-	if (TargetReceiptFilePath.StartsWith(ProjectDirPrefix))
-	{
-		const FStringView TargetReceiptSuffix = TargetReceiptFilePath.RightChop(ProjectDirPrefix.Len());
-		FinalTargetReceiptFilePath = FPaths::ConvertRelativePathToFull(FPaths::ProjectDir() + TargetReceiptSuffix);
-		FinalPublishedVersionFilePath = FinalTargetReceiptFilePath + TEXT(".BuildVersion");
-		FinalGeneratedVersionFilePath = FPaths::ConvertRelativePathToFull(FPaths::Combine(FPaths::ProjectSavedDir(), TEXT("GeneratedWorkerBuildVersions")) + TargetReceiptSuffix + TEXT(".BuildVersion"));
-	}
-	else if (TargetReceiptFilePath.StartsWith(EngineDirPrefix))
-	{
-		// Target receipts declared to be in the engine directory may end up in the project directory if they were compiled
-		// with a uproject file on the commandline.  For this reason we allow for the possibility that the target receipt path
-		// can be tried with both the true engine directory or the engine directory treated as the project directory.
-		// Pick whichever is newer (more recently modified).
-		const FStringView TargetReceiptSuffix = TargetReceiptFilePath.RightChop(EngineDirPrefix.Len());
-		FString ReceiptEnginePath = FPaths::EngineDir() + TargetReceiptSuffix;
-		FString ReceiptProjectPath = FPaths::ProjectDir() + TargetReceiptSuffix;
-		FDateTime ReceiptEngineTimeStamp;
-		FDateTime ReceiptProjectTimeStamp;
-		FileManager.GetTimeStampPair(*ReceiptEnginePath, *ReceiptProjectPath, ReceiptEngineTimeStamp, ReceiptProjectTimeStamp);
-		if (ReceiptProjectTimeStamp > ReceiptEngineTimeStamp)
+	for (const FRuntimeDependency& RuntimeDependency : TargetReceipt.RuntimeDependencies)
+	{
+		if (RuntimeDependency.Path.EndsWith(TEXT(".uproject")))
 		{
-			FinalTargetReceiptFilePath = FPaths::ConvertRelativePathToFull(ReceiptProjectPath);
-			FinalPublishedVersionFilePath = FinalTargetReceiptFilePath + TEXT(".BuildVersion");
-			FinalGeneratedVersionFilePath = FPaths::ConvertRelativePathToFull(FPaths::Combine(FPaths::ProjectSavedDir(), TEXT("GeneratedWorkerBuildVersions")) + TargetReceiptSuffix + TEXT(".BuildVersion"));
+			// uproject files are not needed by build workers.
+			continue;
 		}
-		else
+		if (!TryAddFile(RuntimeDependency.Path))
 		{
-			FinalTargetReceiptFilePath = FPaths::ConvertRelativePathToFull(ReceiptEnginePath);
-			FinalPublishedVersionFilePath = FinalTargetReceiptFilePath + TEXT(".BuildVersion");
-			FinalGeneratedVersionFilePath = FPaths::ConvertRelativePathToFull(FPaths::Combine(FPaths::EngineSavedDir(), TEXT("GeneratedWorkerBuildVersions")) + TargetReceiptSuffix + TEXT(".BuildVersion"));
+			return {};
 		}
 	}
-	else
-	{
-		UE_LOGFMT(LogTargetReceiptBuildWorker, Warning, "Unknown prefix for target receipt path: '{TargetReceipt}'", TargetReceiptFilePath);
-		return false;
-	}
-	
-	if (!FPaths::FileExists(FinalTargetReceiptFilePath))
-	{
-		// This case should return without emitting a warning.
-		return false;
-	}
 
-	FTargetReceipt TargetReceipt;
-	if (!TargetReceipt.Read(FinalTargetReceiptFilePath, false))
-	{
-		UE_LOGFMT(LogTargetReceiptBuildWorker, Warning, "Target receipt failed to parse: '{TargetReceipt}'", TargetReceiptFilePath);
-		return false;
-	}
-
-	if (bAbortRequested.load(std::memory_order_relaxed))
-	{
-		return false;
-	}
-
-	Name = FUtf8String(TargetReceipt.TargetName);
-	Platform = FUtf8String(TargetReceipt.Platform);
-	if (!TryAddExecutablePath(TargetReceipt.Launch))
-	{
-		return false;
-	}
+	// Scan the additional properties for build worker versions.
+	FGuid BuildWorkerReceiptVersion;
 
-	FGuid ReceiptInlineBuildSystemVersion;
-	TArray<TPair<FUtf8String, FGuid>> ReceiptInlineFunctionVersions;
-	bool bUsableBuildWorker = false;
 	for (const FReceiptProperty& Property : TargetReceipt.AdditionalProperties)
 	{
-		if (Property.Name == TEXT("DerivedDataBuildWorkerReceiptVersion"))
+		if (Property.Name == TEXTVIEW("DerivedDataBuildWorkerReceiptVersion"))
 		{
-			bUsableBuildWorker = FGuid(Property.Value) == WorkerReceiptVersion;
+			if (!FGuid::Parse(Property.Value, BuildWorkerReceiptVersion))
+			{
+				UE_LOGFMT(LogTargetReceiptBuildWorker, Warning,
+					"Invalid build worker receipt version '{Version}' in target receipt for '{Name}' at path '{TargetReceipt}'.",
+					Property.Value, Name, TargetReceiptPath);
+				return {};
+			}
 		}
-		else if (Property.Name == TEXT("DerivedDataBuildWorkerBuildSystemVersion"))
+		else if (Property.Name == TEXTVIEW("DerivedDataBuildWorkerBuildSystemVersion"))
 		{
-			if (!FGuid::Parse(Property.Value, ReceiptInlineBuildSystemVersion))
+			if (!FGuid::Parse(Property.Value, BuildSystemVersion))
 			{
-				UE_LOG(LogTargetReceiptBuildWorker, Warning, TEXT("Target receipt contains inline build system version guid that could not be parsed: '%s'"), *Property.Value);
-				return false;
+				UE_LOGFMT(LogTargetReceiptBuildWorker, Warning,
+					"Invalid build worker build system version '{Version}' in target receipt for '{Name}' at path '{TargetReceipt}'.",
+					Property.Value, Name, TargetReceiptPath);
+				return {};
 			}
 		}
-		else if (Property.Name == TEXT("DerivedDataBuildWorkerFunctionVersion"))
+		else if (Property.Name == TEXTVIEW("DerivedDataBuildWorkerFunctionVersion"))
 		{
-			TArray<FString> FunctionComponents;
-			if (2 != Property.Value.ParseIntoArray(FunctionComponents, TEXT(":")))
+			TArray<FStringView, TInlineAllocator<2>> FunctionComponents;
+			UE::String::ParseTokens(Property.Value, TEXT(':'), FunctionComponents);
+			if (FunctionComponents.Num() != 2)
 			{
-				UE_LOG(LogTargetReceiptBuildWorker, Warning, TEXT("Target receipt contains inline function version declaration that could not be parsed: '%s'"), *Property.Value);
-				return false;
+				UE_LOGFMT(LogTargetReceiptBuildWorker, Warning,
+					"Invalid build worker function name and version '{Version}' in target receipt for '{Name}' at path '{TargetReceipt}'.",
+					Property.Value, Name, TargetReceiptPath);
+				return {};
 			}
 
-			FGuid VersionGuid;
-			if (!FGuid::Parse(FunctionComponents[1], VersionGuid))
+			FGuid FunctionVersion;
+			if (!FGuid::Parse(FunctionComponents[1], FunctionVersion))
 			{
-				UE_LOG(LogTargetReceiptBuildWorker, Warning, TEXT("Target receipt contains inline function version guid that could not be parsed: '%s'"), *Property.Value);
-				return false;
+				UE_LOGFMT(LogTargetReceiptBuildWorker, Warning,
+					"Invalid build worker function version '{Version}' in target receipt for '{Name}' at path '{TargetReceipt}'.",
+					FunctionComponents[1], Name, TargetReceiptPath);
+				return {};
 			}
 
-			ReceiptInlineFunctionVersions.Emplace(FunctionComponents[0], MoveTemp(VersionGuid));
+			Functions.Emplace(FunctionComponents[0], MoveTemp(FunctionVersion));
 		}
-		else if (Property.Name == TEXT("UnstagedRuntimeDependency"))
+		else if (Property.Name == TEXTVIEW("UnstagedRuntimeDependency"))
 		{
-			if (!TryAddExecutablePath(Property.Value))
+			if (!TryAddFile(Property.Value))
 			{
-				UE_LOG(LogTargetReceiptBuildWorker, Warning, TEXT("Target receipt referenced an unstaged dependency that was not usable: '%s'"), *Property.Value);
-				return false;
+				UE_LOGFMT(LogTargetReceiptBuildWorker, Warning,
+					"Unusable unstaged dependency '{Path}' referenced by '{Name}' at '{TargetReceipt}'.", Property.Value, Name, TargetReceiptPath);
+				return {};
 			}
 		}
 	}
 
-	if (!bUsableBuildWorker)
+	if (BuildWorkerReceiptVersion != CurrentWorkerReceiptVersion)
 	{
-		return false;
+		UE_LOGFMT(LogTargetReceiptBuildWorker, Verbose,
+			"Ignoring target with no build worker receipt version for '{Name}' at path '{TargetReceipt}'...", Name, TargetReceiptPath);
+		return {};
 	}
 
-	for (const FBuildProduct& BuildProduct : TargetReceipt.BuildProducts)
+	if (Params.ShouldAbort && Params.ShouldAbort())
 	{
-		if (BuildProduct.Type == TEXTVIEW("SymbolFile") ||
-			BuildProduct.Path == TargetReceipt.Launch)
-		{
-			continue;
-		}
-		if (!TryAddExecutablePath(BuildProduct.Path))
-		{
-			UE_LOG(LogTargetReceiptBuildWorker, Warning, TEXT("Target receipt referenced a build product that was not usable: '%s'"), *BuildProduct.Path);
-			return false;
-		}
+		return {};
 	}
 
-	for (const FRuntimeDependency& RuntimeDependency : TargetReceipt.RuntimeDependencies)
+	// Find the most recently modified file referenced by the target receipt.
+	MostRecentFileModTime = Algo::MaxElementBy(LocalPathToMetaData,
+		Projection(&TPair<UE::FSharedString, FFileMetaData>::Value, &FFileMetaData::ModTime))->Value.ModTime;
+
+	if (WorkerPackageModTime > MostRecentFileModTime && !Params.bAlwaysBuild)
 	{
-		if (RuntimeDependency.Path.EndsWith(TEXT(".uproject")))
+		if (TOptional<UE::DerivedData::FBuildWorker> Worker = Registry.LoadWorker(ActiveWorkerPackagePath))
 		{
-			// uproject files are not needed for build workers
-			continue;
-		}
-		if (!TryAddExecutablePath(RuntimeDependency.Path))
-		{
-			UE_LOG(LogTargetReceiptBuildWorker, Warning, TEXT("Target receipt referenced a runtime dependency that was not usable: '%s'"), *RuntimeDependency.Path);
-			return false;
+			UE_LOGFMT(LogTargetReceiptBuildWorker, Verbose,
+				"Files referenced by the target receipt have not been modified since the build worker package was created. "
+				"Using existing build worker package at '{Path}' for '{Name}' at path '{TargetReceipt}'.",
+				ActiveWorkerPackagePath, Name, TargetReceiptPath);
+			return Worker;
 		}
 	}
-	
-	if (bAbortRequested.load(std::memory_order_relaxed))
+
+	if (!BuildSystemVersion.IsValid() && !TryLoadWorkerVersions(ActiveTargetReceiptPath, ActiveWorkerPackagePath))
 	{
-		return false;
+		return {};
 	}
 
-	// Build a sorted list of unique paths because the target receipt may reference the same file in multiple places.
-	// Skip the first path because it is assumed elsewhere to be the launch path.
-	Algo::SortBy(MakeArrayView(ExecutablePaths).RightChop(1), &FWorkerPath::LocalPath);
-	const int32 KeepPaths = Algo::UniqueBy(ExecutablePaths, &FWorkerPath::LocalPath);
-	ExecutablePaths.RemoveAt(KeepPaths, ExecutablePaths.Num() - KeepPaths, EAllowShrinking::No);
-
-	// Open executable files to check for existence and gather hash info for them
-	int32 MetaIndex = 0;
-	ExecutableMeta.AddDefaulted(ExecutablePaths.Num());
-	// Consider the receipt file when computing newest executable modification time
-	FDateTime NewestExecutableModificationTime = FileManager.GetTimeStamp(*FinalTargetReceiptFilePath);
-	for (const FWorkerPath& ExecutablePath : ExecutablePaths)
+	if (Params.ShouldAbort && Params.ShouldAbort())
 	{
-		FWorkerPathMeta& Meta = ExecutableMeta[MetaIndex];
+		return {};
+	}
 
-		TStringBuilder<256> LocalExecutablePath(InPlace, ExecutablePath.LocalPath);
+	UE_LOGFMT(LogTargetReceiptBuildWorker, Log, "Building a build worker package from {FileCount} files referenced by '{TargetReceipt}'...",
+		LocalPathToMetaData.Num(), TargetReceiptPath);
 
-		if (TUniquePtr<FArchive> Ar{FileManager.CreateFileReader(*LocalExecutablePath, FILEREAD_Silent)})
+	// Load the referenced files from disk.
+	for (TPair<UE::FSharedString, FFileMetaData>& It : LocalPathToMetaData)
+	{
+		UE_LOGFMT(LogTargetReceiptBuildWorker, Verbose, "Loading file '{Path}' for '{Name}' at path '{TargetReceipt}'...",
+			It.Value.LocalPath, Name, TargetReceiptPath);
+		if (TUniquePtr<FArchive> Ar{FileManager.CreateFileReader(*It.Value.LocalPath, FILEREAD_Silent)})
 		{
-			FDateTime ModificationTime = FileManager.GetTimeStamp(*LocalExecutablePath);
-			if (ModificationTime > NewestExecutableModificationTime)
+			const int64 TotalSize = FPlatformMath::Max(0, Ar->TotalSize());
+			FUniqueBuffer Data = FUniqueBuffer::Alloc(uint64(TotalSize));
+			Ar->Serialize(Data.GetData(), TotalSize);
+			if (!Ar->Close())
 			{
-				NewestExecutableModificationTime = ModificationTime;
+				UE_LOGFMT(LogTargetReceiptBuildWorker, Warning, "Failed to read file '{Path}' for '{Name}' at path '{TargetReceipt}'.",
+					It.Value.LocalPath, Name, TargetReceiptPath);
+				return {};
 			}
-			const int64 TotalSize = Ar->TotalSize();
-			FUniqueBuffer MutableBuffer = FUniqueBuffer::Alloc(uint64(TotalSize));
-			Ar->Serialize(MutableBuffer.GetData(), TotalSize);
-			Meta.Key = FIoHash::HashBuffer(MutableBuffer.GetView());
-			Meta.Value = uint64(TotalSize);
-			Ar->Close();
+			It.Value.RawData = Data.MoveToShared();
 		}
 		else
 		{
-			UE_LOG(LogTargetReceiptBuildWorker, Warning, TEXT("Executable file not present: '%s'"), *LocalExecutablePath);
-			return false;
+			UE_LOGFMT(LogTargetReceiptBuildWorker, Warning, "Failed to open file '{Path}' for '{Name}' at path '{TargetReceipt}'.",
+				It.Value.LocalPath, Name, TargetReceiptPath);
+			return {};
+		}
+
+		if (Params.ShouldAbort && Params.ShouldAbort())
+		{
+			return {};
 		}
-		++MetaIndex;
 	}
+	const uint64 TotalFileSize = Algo::TransformAccumulate(LocalPathToMetaData,
+		Projection(&TPair<UE::FSharedString, FFileMetaData>::Value, &FFileMetaData::RawData, &FSharedBuffer::GetSize), uint64(0));
 
-	if (ReceiptInlineBuildSystemVersion.IsValid())
+	// Hash and optionally compress the files.
+	ParallelFor(TEXT("CompressBuildWorker"), LocalPathToMetaData.Num(), /*MinBatchSize*/ 1, [this](int32 Index)
+	{
+		FFileMetaData& MetaData = LocalPathToMetaData.Get(FSetElementId::FromInteger(Index)).Value;
+		if (Params.bSkipCompression)
+		{
+			MetaData.Data = FCompressedBuffer::Compress(MetaData.RawData, ECompressedBufferCompressor::NotSet, ECompressedBufferCompressionLevel::None);
+		}
+		else
+		{
+			MetaData.Data = FCompressedBuffer::Compress(MetaData.RawData, ECompressedBufferCompressor::Kraken, ECompressedBufferCompressionLevel::Normal);
+		}
+		MetaData.RawData.Reset();
+	});
+
+	// Build the build worker package.
+	UE::DerivedData::FBuildWorkerBuilder Builder = Registry.CreateWorker();
+	Builder.SetName(WriteToUtf8String<64>(Name));
+	Builder.SetPath(WriteToUtf8String<256>(LocalPathToMetaData[LocalLaunchPath].RemotePath));
+	Builder.SetHostPlatform(WriteToUtf8String<64>(HostPlatform));
+	Builder.SetBuildSystemVersion(BuildSystemVersion);
+	for (const TTuple<FString, FGuid>& Function : Functions)
 	{
-		// Accept version information declared inline in the target receipt.
-		BuildSystemVersion = MoveTemp(ReceiptInlineBuildSystemVersion);
-		FunctionVersions = MoveTemp(ReceiptInlineFunctionVersions);
+		Builder.AddFunction(WriteToUtf8String<64>(Function.Key), Function.Value);
 	}
-	else
+	for (const TPair<UE::FSharedString, FFileMetaData>& It : LocalPathToMetaData)
 	{
-		// Gather system and function version information from the process
-		FString VersionFilePath;
-		FCbObject VersionObject;
-		if (FPaths::FileExists(FinalPublishedVersionFilePath))
+		if (It.Value.bExecutable)
 		{
-			VersionFilePath = FinalPublishedVersionFilePath;
-
-			if (TUniquePtr<FArchive> Ar{IFileManager::Get().CreateFileReader(*FinalPublishedVersionFilePath, FILEREAD_Silent)})
-			{
-				*Ar << VersionObject;
-			}
-			else
-			{
-				UE_LOG(LogTargetReceiptBuildWorker, Warning, TEXT("Could not read published version file: '%s'"), *TargetReceipt.Launch);
-				return false;
-			}
+			Builder.AddExecutable(WriteToUtf8String<256>(It.Value.RemotePath), It.Value.Data);
 		}
 		else
 		{
-			VersionFilePath = FinalGeneratedVersionFilePath;
+			Builder.AddFile(WriteToUtf8String<256>(It.Value.RemotePath), It.Value.Data);
+		}
+	}
+	for (const TPair<FString, FString>& Env : Environment)
+	{
+		Builder.SetEnvironment(WriteToUtf8String<64>(Env.Key), WriteToUtf8String<192>(Env.Value));
+	}
 
-			if (TargetReceipt.Platform != FPlatformProcess::GetBinariesSubdirectory())
-			{
-				UE_LOG(LogTargetReceiptBuildWorker, Log, TEXT("Target lacked published version information and isn't made for the current platform: '%s'"), *TargetReceipt.Launch);
-				return false;
-			}
+	// Build workers require this directory to exist for correct operation.
+	Builder.AddDirectory(WriteToUtf8String<64>("Engine/Binaries/", HostPlatform));
 
-			FDateTime VersionFileModificationTime = FileManager.GetTimeStamp(*FinalGeneratedVersionFilePath);
-			if (VersionFileModificationTime < NewestExecutableModificationTime)
-			{
-				FString CommandLine = TEXT("-Version=\"") + FinalGeneratedVersionFilePath + TEXT("\"");
-				FProcHandle WorkerProcHandle = FPlatformProcess::CreateProc(*FPaths::ConvertRelativePathToFull(FString(ExecutablePaths[0].LocalPath)), *CommandLine, true, true, true, nullptr, 0, nullptr, nullptr);
-				if (!WorkerProcHandle.IsValid())
-				{
-					UE_LOG(LogTargetReceiptBuildWorker, Warning, TEXT("Failed to launch worker: '%s'"), *TargetReceipt.Launch);
-					return false;
-				}
+	UE::DerivedData::FBuildWorker Worker = Builder.Build(ActiveWorkerPackagePath);
 
-				uint64 WorkerStartTime = FPlatformTime::Cycles64();
-				while (FPlatformProcess::IsProcRunning(WorkerProcHandle))
-				{
-					FPlatformProcess::Sleep(0.1f);
-					if (bAbortRequested.load(std::memory_order_relaxed) || (FPlatformTime::ToSeconds64(FPlatformTime::Cycles64()-WorkerStartTime) > 10))
-					{
-						UE_LOG(LogTargetReceiptBuildWorker, Warning, TEXT("Aborted launch of worker: '%s'"), *TargetReceipt.Launch);
-						FPlatformProcess::TerminateProc(WorkerProcHandle, true);
-						FPlatformProcess::CloseProc(WorkerProcHandle);
-						return false;
-					}
-				}
+	const double EndTime = FPlatformTime::Seconds();
 
-				FPlatformProcess::CloseProc(WorkerProcHandle);
+	const int64 PackageSize = FileManager.FileSize(*ActiveWorkerPackagePath);
+	UE_LOGFMT(LogTargetReceiptBuildWorker, Log, "Created build worker package from '{TargetReceipt}' in {Time} seconds. "
+		"Package is {PackageSizeMB} MiB. Input is {FileSizeMB} MiB in {FileCount} files.",
+		TargetReceiptPath, EndTime - StartTime, PackageSize / 1024 / 1024, TotalFileSize / 1024 / 1024, LocalPathToMetaData.Num());
 
-				if (!FPaths::FileExists(FinalGeneratedVersionFilePath) || (FileManager.GetTimeStamp(*FinalGeneratedVersionFilePath) < NewestExecutableModificationTime))
-				{
-					UE_LOG(LogTargetReceiptBuildWorker, Warning, TEXT("Worker did not produce version file: '%s'"), *TargetReceipt.Launch);
-					return false;
-				}
-			}
+	return Worker;
+}
 
-			if (TUniquePtr<FArchive> Ar{IFileManager::Get().CreateFileReader(*FinalGeneratedVersionFilePath, FILEREAD_Silent)})
+bool FTargetReceiptBuildWorkerBuilder::TryFindTargetReceipt(
+	FStringBuilderBase& OutActiveTargetReceiptPath,
+	FStringBuilderBase& OutActiveWorkerPackagePath)
+{
+	if (TargetReceiptPath.StartsWith(EngineDirPrefix))
+	{
+		// A target that is expected to be in the engine directory will be found in the project
+		// directory if a project was provided at build time. Check both targets and take the one
+		// with the more recent modification time.
+		const FStringView TargetReceiptRelativePath = TargetReceiptPath.RightChop(EngineDirPrefix.Len() + 1);
+		const FString EngineReceiptPath = FPaths::ConvertRelativePathToFull(EngineDir + TargetReceiptRelativePath);
+		const FString ProjectReceiptPath = FPaths::ConvertRelativePathToFull(ProjectDir + TargetReceiptRelativePath);
+		FDateTime EngineReceiptModTime;
+		FDateTime ProjectReceiptModTime;
+		FileManager.GetTimeStampPair(*EngineReceiptPath, *ProjectReceiptPath, EngineReceiptModTime, ProjectReceiptModTime);
+
+		FString ActiveSavedDir;
+		if (ProjectReceiptModTime > EngineReceiptModTime)
+		{
+			OutActiveTargetReceiptPath = ProjectReceiptPath;
+			ActiveSavedDir = FPaths::ProjectSavedDir();
+		}
+		else
+		{
+			OutActiveTargetReceiptPath = EngineReceiptPath;
+			ActiveSavedDir = FPaths::EngineSavedDir();
+		}
+
+		if (OutActiveWorkerPackagePath.Len() == 0)
+		{
+			const FString WorkerPackageRelativePath =
+				FPaths::Combine(TEXTVIEW("BuildWorkers"), FPathViews::GetBaseFilenameWithPath(TargetReceiptRelativePath)) + TEXTVIEW(".worker");
+			if (EngineReceiptModTime == FDateTime::MinValue() && ProjectReceiptModTime == FDateTime::MinValue())
 			{
-				*Ar << VersionObject;
+				// When neither target receipt exists, check which build worker package is newer.
+				const FString EngineWorkerPackagePath = FPaths::ConvertRelativePathToFull(FPaths::Combine(FPaths::EngineSavedDir(), WorkerPackageRelativePath));
+				const FString ProjectWorkerPackagePath = FPaths::ConvertRelativePathToFull(FPaths::Combine(FPaths::ProjectSavedDir(), WorkerPackageRelativePath));
+				FDateTime EngineWorkerPackageModTime;
+				FDateTime ProjectWorkerPackageModTime;
+				FileManager.GetTimeStampPair(*EngineWorkerPackagePath, *ProjectWorkerPackagePath, EngineWorkerPackageModTime, ProjectWorkerPackageModTime);
+				if (ProjectWorkerPackageModTime > EngineWorkerPackageModTime)
+				{
+					OutActiveTargetReceiptPath = ProjectReceiptPath;
+					OutActiveWorkerPackagePath = ProjectWorkerPackagePath;
+				}
+				else
+				{
+					OutActiveTargetReceiptPath = EngineReceiptPath;
+					OutActiveWorkerPackagePath = EngineWorkerPackagePath;
+				}
 			}
 			else
 			{
-				UE_LOG(LogTargetReceiptBuildWorker, Warning, TEXT("Worker did not produce readable version file: '%s'"), *TargetReceipt.Launch);
-				return false;
+				OutActiveWorkerPackagePath = FPaths::ConvertRelativePathToFull(FPaths::Combine(ActiveSavedDir, WorkerPackageRelativePath));
 			}
 		}
+		return true;
+	}
 
-		BuildSystemVersion = VersionObject["BuildSystemVersion"].AsUuid();
-		if (!BuildSystemVersion.IsValid())
+	if (TargetReceiptPath.StartsWith(ProjectDirPrefix))
+	{
+		const FStringView TargetReceiptRelativePath = TargetReceiptPath.RightChop(ProjectDirPrefix.Len() + 1);
+		OutActiveTargetReceiptPath = FPaths::ConvertRelativePathToFull(ProjectDir + TargetReceiptRelativePath);
+		if (OutActiveWorkerPackagePath.Len() == 0)
 		{
-			UE_LOG(LogTargetReceiptBuildWorker, Warning, TEXT("Invalid build system version in version file: '%s'"), *VersionFilePath);
-			return false;
+			OutActiveWorkerPackagePath = FPaths::ConvertRelativePathToFull(FPaths::Combine(FPaths::ProjectSavedDir(),
+				TEXTVIEW("BuildWorkers"), FPathViews::GetBaseFilenameWithPath(TargetReceiptRelativePath)) + TEXTVIEW(".worker"));
 		}
+		return true;
+	}
 
-		FCbArrayView FunctionsArrayView = VersionObject["Functions"].AsArrayView();
-		for (FCbFieldView FunctionFieldView : FunctionsArrayView)
-		{
-			FCbObjectView FunctionObjectView = FunctionFieldView.AsObjectView();
-			FunctionVersions.Emplace(FunctionObjectView["Name"].AsString(), FunctionObjectView["Version"].AsUuid());
-		}
+	if (TargetReceiptPath.StartsWith(TEXTVIEW("$(")))
+	{
+		UE_LOGFMT(LogTargetReceiptBuildWorker, Warning, "Unknown prefix for target receipt path '{TargetReceipt}'.", TargetReceiptPath);
+		return false;
+	}
+
+	if (OutActiveWorkerPackagePath.Len() == 0)
+	{
+		UE_LOGFMT(LogTargetReceiptBuildWorker, Error, "The build worker package path must be provided because the "
+			"target receipt path '{TargetReceipt}' is not prefixed with '{EngineDirPrefix}' or '{ProjectDirPrefix}'.",
+			TargetReceiptPath, EngineDirPrefix, ProjectDirPrefix);
+		return false;
 	}
 
+	OutActiveTargetReceiptPath = TargetReceiptPath;
 	return true;
 }
 
-void FTargetReceiptBuildWorkerLoader::AbortLoad()
+FTargetReceiptBuildWorkerBuilder::FFileMetaData* FTargetReceiptBuildWorkerBuilder::TryAddFile(FStringView Path)
 {
-	bAbortRequested.store(true, std::memory_order_relaxed);
-}
+	if (Path.IsEmpty())
+	{
+		UE_LOGFMT(LogTargetReceiptBuildWorker, Warning,
+			"Empty path found in target receipt for '{Name}' at path '{TargetReceipt}'.", Name, TargetReceiptPath);
+		return {};
+	}
 
-class FTargetReceiptBuildWorkerResolver final : public UE::DerivedData::IBuildWorkerResolver
-{
-public:
-	explicit FTargetReceiptBuildWorkerResolver(TMap<FIoHash, FString>&& InHashToPath)
-		: HashToPath(MoveTemp(InHashToPath))
+	TStringBuilder<256> LocalPath;
+	TStringBuilder<256> RemotePath;
+
+	if (Path.StartsWith(EngineDirPrefix))
 	{
+		FPathViews::Append(LocalPath, EngineDir, Path.RightChop(EngineDirPrefix.Len() + 1));
+		FPathViews::Append(RemotePath, ANSITEXTVIEW("Engine"), Path.RightChop(EngineDirPrefix.Len() + 1));
 	}
+	else if (Path.StartsWith(ProjectDirPrefix))
+	{
+		FPathViews::Append(LocalPath, ProjectDir, Path.RightChop(ProjectDirPrefix.Len() + 1));
+		FPathViews::Append(RemotePath, RelativeProjectDir, Path.RightChop(ProjectDirPrefix.Len() + 1));
+	}
+	else
+	{
+		if (!Path.StartsWith(TEXTVIEW("$(")))
+		{
+			UE_LOGFMT(LogTargetReceiptBuildWorker, Warning,
+				"Absolute or unprefixed path '{Path}' found in target receipt for '{Name}' at path '{TargetReceipt}'.",
+				Path, Name, TargetReceiptPath);
+			return {};
+		}
 
-	void Resolve(TConstArrayView<FIoHash> RawHashes, UE::DerivedData::IRequestOwner& Owner, UE::DerivedData::FOnBuildWorkerResolved&& OnResolved) final
+		const int32 EndVariableIndex = Path.Find(TEXTVIEW(")"));
+		if (EndVariableIndex == INDEX_NONE)
+		{
+			UE_LOGFMT(LogTargetReceiptBuildWorker, Warning,
+				"Expected ')' following '$(' in path '{Path}' in target receipt for '{Name}' at path '{TargetReceipt}'.",
+				Path, Name, TargetReceiptPath);
+			return {};
+		}
+
+		const FStringView RelativePath = Path.RightChop(EndVariableIndex + 2);
+		const FStringView EnvVarName = Path.Mid(2, EndVariableIndex - 2);
+		const FString EnvVarValue = FPlatformMisc::GetEnvironmentVariable(*WriteToString<64>(EnvVarName));
+		if (EnvVarValue.IsEmpty())
+		{
+			UE_LOGFMT(LogTargetReceiptBuildWorker, Warning,
+				"Unknown environment variable '{EnvVar}' in path '{Path}' in target receipt for '{Name}' at path '{TargetReceipt}'.",
+				EnvVarName, Path, Name, TargetReceiptPath);
+			return {};
+		}
+
+		// Depending on an environment variable to locate the file now creates an implicit
+		// dependency on the environment variable when executing the worker later.
+		// Files within this local directory are remapped to a directory within the worker root
+		// that is named the same as the environment variable. That path needs to be added to the
+		// environment of the worker to allow it to resolve these files when it executes.
+		FPathViews::Append(RemotePath, "..", "..", "..", EnvVarName);
+		Environment.Emplace(EnvVarName, RemotePath);
+		FPathViews::Append(RemotePath, RelativePath);
+		FPathViews::Append(LocalPath, EnvVarValue, RelativePath);
+	}
+
+	FDateTime ModTime = FileManager.GetTimeStamp(*LocalPath);
+	if (ModTime == FDateTime::MinValue())
 	{
-		TArray<FCompressedBuffer> Files;
-		for (const FIoHash& RawHash : RawHashes)
+		UE_LOGFMT(LogTargetReceiptBuildWorker, Warning,
+			"Missing file '{LocalPath}' referenced by target receipt for '{Name}' at path '{TargetReceipt}'.", LocalPath, Name, TargetReceiptPath);
+		return {};
+	}
+
+	FFileMetaData MetaData;
+	MetaData.LocalPath = LocalPath;
+	MetaData.RemotePath = RemotePath;
+	MetaData.ModTime = ModTime;
+	return &LocalPathToMetaData.FindOrAdd(MetaData.LocalPath, MetaData);
+}
+
+bool FTargetReceiptBuildWorkerBuilder::TryLoadWorkerVersions(FStringView ActiveTargetReceiptPath, FStringView ActiveWorkerPackagePath)
+{
+	FCbObject VersionObject;
+
+	TStringBuilder<256> VersionPath(InPlace, ActiveTargetReceiptPath, TEXTVIEW(".BuildVersion"));
+	if (!FileManager.FileExists(*VersionPath))
+	{
+		if (FPlatformProcess::GetBinariesSubdirectory() != HostPlatform)
+		{
+			UE_LOGFMT(LogTargetReceiptBuildWorker, Log,
+				"The build worker is for platform '{HostPlatform}' and does not have published versions for '{Name}' at path '{TargetReceipt}'.",
+				HostPlatform, Name, TargetReceiptPath);
+			return false;
+		}
+
+		VersionPath = ActiveWorkerPackagePath;
+		VersionPath.Append(TEXTVIEW(".BuildVersion"));
+
+		FDateTime VersionFileModificationTime = FileManager.GetTimeStamp(*VersionPath);
+		if (VersionFileModificationTime < MostRecentFileModTime)
 		{
-			if (const FString* Path = HashToPath.Find(RawHash))
+			UE_LOGFMT(LogTargetReceiptBuildWorker, Verbose,
+				"Launching '{LaunchPath}' to report versions for '{Name}'...", LocalLaunchPath, Name);
+
+			TStringBuilder<256> CommandLineParams(InPlace, TEXTVIEW("-Version=\""), VersionPath, TEXTVIEW("\""));
+			FProcHandle WorkerProcHandle = FPlatformProcess::CreateProc(
+				*FPaths::ConvertRelativePathToFull(FString(LocalLaunchPath)),
+				*CommandLineParams,
+				/*bLaunchDetached*/ true,
+				/*bLaunchHidden*/ true,
+				/*bLaunchReallyHidden*/ true,
+				/*OutProcessID*/ nullptr,
+				/*PriorityModifier*/ 0,
+				/*OptionalWorkingDirectory*/ nullptr,
+				/*PipeWriteChild*/ nullptr);
+			if (!WorkerProcHandle.IsValid())
+			{
+				UE_LOGFMT(LogTargetReceiptBuildWorker, Warning,
+					"Failed to launch build worker '{LaunchPath}' to report versions for '{Name}'.", LocalLaunchPath, Name);
+				return false;
+			}
+
+			uint64 WorkerStartTime = FPlatformTime::Cycles64();
+			while (FPlatformProcess::IsProcRunning(WorkerProcHandle))
 			{
-				if (TUniquePtr<FArchive> Ar{IFileManager::Get().CreateFileReader(**Path, FILEREAD_Silent)})
+				FPlatformProcess::Sleep(0.1f);
+				if ((Params.ShouldAbort && Params.ShouldAbort()) || (FPlatformTime::ToSeconds64(FPlatformTime::Cycles64() - WorkerStartTime) > 10))
 				{
-					const int64 TotalSize = Ar->TotalSize();
-					FUniqueBuffer MutableBuffer = FUniqueBuffer::Alloc(uint64(TotalSize));
-					Ar->Serialize(MutableBuffer.GetData(), TotalSize);
-					if (Ar->Close())
-					{
-						Files.Emplace(FCompressedBuffer::Compress(
-							MutableBuffer.MoveToShared(),
-							ECompressedBufferCompressor::NotSet,
-							ECompressedBufferCompressionLevel::None));
-					}
+					UE_LOGFMT(LogTargetReceiptBuildWorker, Warning,
+						"Aborted build worker '{LaunchPath}' version reporting for '{Name}'.", LocalLaunchPath, Name);
+					FPlatformProcess::TerminateProc(WorkerProcHandle, /*bKillTree*/ true);
+					FPlatformProcess::CloseProc(WorkerProcHandle);
+					return false;
 				}
 			}
+
+			FPlatformProcess::CloseProc(WorkerProcHandle);
+
+			if (FileManager.GetTimeStamp(*VersionPath) < MostRecentFileModTime)
+			{
+				UE_LOGFMT(LogTargetReceiptBuildWorker, Warning,
+					"The build worker '{LaunchPath}' did not produce a version file for '{Name}'.", LocalLaunchPath, Name);
+				return false;
+			}
 		}
-		OnResolved({Files, Files.Num() == RawHashes.Num() ? UE::DerivedData::EStatus::Ok : UE::DerivedData::EStatus::Error});
 	}
 
-private:
-	TMap<FIoHash, FString> HashToPath;
-};
+	if (TUniquePtr<FArchive> Ar{IFileManager::Get().CreateFileReader(*VersionPath, FILEREAD_Silent)})
+	{
+		*Ar << VersionObject;
+	}
+	else
+	{
+		UE_LOGFMT(LogTargetReceiptBuildWorker, Warning,
+			"Failed to load build worker version file '{VersionPath}' for '{Name}'.", VersionPath, Name);
+		return false;
+	}
 
-UE::DerivedData::FBuildWorker FTargetReceiptBuildWorkerLoader::Build(UE::DerivedData::FBuildWorkerBuilder& Builder) const
-{
-	Builder.SetName(Name);
-	Builder.SetPath(ExecutablePaths[0].RemotePath);
-	Builder.SetHostPlatform(Platform);
+	BuildSystemVersion = VersionObject["BuildSystemVersion"].AsUuid();
+	if (!BuildSystemVersion.IsValid())
+	{
+		UE_LOGFMT(LogTargetReceiptBuildWorker, Warning,
+			"Invalid build worker build system version in published version file '{VersionPath}' for '{Name}'.", VersionPath, Name);
+		return false;
+	}
 
-	Builder.SetBuildSystemVersion(BuildSystemVersion);
+	for (FCbFieldView FunctionField : VersionObject["Functions"])
+	{
+		Functions.Emplace(FunctionField["Name"].AsString(), FunctionField["Version"].AsUuid());
+	}
 
-	for (const TPair<FUtf8String, FGuid>& FunctionVersion : FunctionVersions)
+	return true;
+}
+
+TOptional<UE::DerivedData::FBuildWorker> CreateBuildWorkerFromTargetReceipt(
+	FStringView TargetReceiptPath,
+	UE::DerivedData::IBuildWorkerRegistry& Registry,
+	const FTargetReceiptBuildWorkerParams& Params)
+{
+	FTargetReceiptBuildWorkerBuilder Builder(TargetReceiptPath, Params);
+	return Builder.Build(Registry);
+}
+
+#if WITH_EDITOR
+class FTargetReceiptBuildWorkerFactory final : public UE::DerivedData::IBuildWorkerFactory
+{
+public:
+	explicit FTargetReceiptBuildWorkerFactory(FStringView TargetReceiptFilePath)
+		: TargetReceiptPath(TargetReceiptFilePath)
 	{
-		Builder.AddFunction(FunctionVersion.Key, FunctionVersion.Value);
+		IModularFeatures::Get().RegisterModularFeature(UE::DerivedData::IBuildWorkerFactory::FeatureName, this);
 	}
 
-	for (int32 ExecutableIndex = 0; ExecutableIndex < ExecutablePaths.Num(); ++ExecutableIndex)
+	~FTargetReceiptBuildWorkerFactory() final
 	{
-		Builder.AddExecutable(ExecutablePaths[ExecutableIndex].RemotePath, ExecutableMeta[ExecutableIndex].Key, ExecutableMeta[ExecutableIndex].Value);
+		IModularFeatures::Get().UnregisterModularFeature(UE::DerivedData::IBuildWorkerFactory::FeatureName, this);
 	}
 
-	for (const TPair<FUtf8String, FUtf8String>& EnvironmentVariable : EnvironmentVariables)
+	void CreateWorkers(UE::DerivedData::IBuildWorkerRegistry& Registry) final
 	{
-		Builder.SetEnvironment(EnvironmentVariable.Key, EnvironmentVariable.Value);
+		CreateWorkerTask = UE::Tasks::Launch(TEXT("TargetReceiptBuildWorkerFactory"), [this, &Registry]
+		{
+			const auto ShouldAbort = [this] { return bAborted.load(std::memory_order_relaxed); };
+			if (TOptional<UE::DerivedData::FBuildWorker> Worker = CreateBuildWorkerFromTargetReceipt(TargetReceiptPath, Registry, {.ShouldAbort = ShouldAbort}))
+			{
+				Registry.RegisterWorker(MoveTemp(*Worker), this);
+			}
+		});
 	}
 
-	TMap<FIoHash, FString> HashToPath;
-	HashToPath.Reserve(ExecutablePaths.Num());
-	for (int32 ExecutableIndex = 0; ExecutableIndex < ExecutablePaths.Num(); ++ExecutableIndex)
+	void AbortCreateWorkers() final
 	{
-		HashToPath.Emplace(ExecutableMeta[ExecutableIndex].Get<FIoHash>(), ExecutablePaths[ExecutableIndex].LocalPath);
+		bAborted.store(true, std::memory_order_relaxed);
+		CreateWorkerTask.Wait();
 	}
-	return Builder.Build(MakeUnique<FTargetReceiptBuildWorkerResolver>(MoveTemp(HashToPath)));
+
+private:
+	UE::FSharedString TargetReceiptPath;
+	UE::Tasks::FTask CreateWorkerTask;
+	std::atomic<bool> bAborted = false;
+};
+#endif // WITH_EDITOR
+
+FTargetReceiptBuildWorker::FTargetReceiptBuildWorker(FStringView TargetReceiptFilePath)
+{
+#if WITH_EDITOR
+	Factory = MakePimpl<FTargetReceiptBuildWorkerFactory>(TargetReceiptFilePath);
+#endif
 }
+
+FTargetReceiptBuildWorker::~FTargetReceiptBuildWorker() = default;
diff --git a/Engine/Source/Developer/DesktopPlatform/Public/TargetReceiptBuildWorker.h b/Engine/Source/Developer/DesktopPlatform/Public/TargetReceiptBuildWorker.h
--- a/Engine/Source/Developer/DesktopPlatform/Public/TargetReceiptBuildWorker.h
+++ b/Engine/Source/Developer/DesktopPlatform/Public/TargetReceiptBuildWorker.h
@@ -3,12 +3,45 @@
 #pragma once
 
 #include "Containers/StringFwd.h"
+#include "Containers/StringView.h"
+#include "Misc/OptionalFwd.h"
+#include "Templates/FunctionWithContext.h"
 #include "Templates/PimplPtr.h"
 
 #define UE_API DESKTOPPLATFORM_API
 
 class FTargetReceiptBuildWorkerFactory;
 
+namespace UE::DerivedData { class FBuildWorker; }
+namespace UE::DerivedData { class IBuildWorkerRegistry; }
+
+struct FTargetReceiptBuildWorkerParams
+{
+	/** When provided, override the default saved builder worker package path with this path. */
+	FStringView WorkerPackagePath;
+	/** When provided, called periodically to check if build worker creation should abort early. */
+	UE::TFunctionWithContext<bool ()> ShouldAbort = nullptr;
+	/** When true, skip timestamp checks and unconditionally create the build worker package. */
+	bool bAlwaysBuild = false;
+	/** When true, include symbol files in the build worker package. */
+	bool bIncludeSymbols = false;
+	/** When true, skips compression of the build worker package attachments. */
+	bool bSkipCompression = false;
+};
+
+/**
+ * Creates a build worker package from the target receipt.
+ *
+ * @param TargetReceiptPath   Path to load the target receipt from.
+ * @param Registry   Registry provided to avoid DesktopPlatform requiring a link dependency on DerivedDataCache.
+ * @param Params   Optional parameters to configure build worker package creation.
+ * @return A build worker for a package saved to the path, or unset on error.
+ */
+[[nodiscard]] UE_API TOptional<UE::DerivedData::FBuildWorker> CreateBuildWorkerFromTargetReceipt(
+	FStringView TargetReceiptPath,
+	UE::DerivedData::IBuildWorkerRegistry& Registry,
+	const FTargetReceiptBuildWorkerParams& Params = {});
+
 /**
  * Globally registers a UE::DerivedData::IBuildWorkerFactory instance that runs an executable built by UnrealBuildTool.
  * UnrealBuildTool provides the executable information through a TargetReceipt file.
diff --git a/Engine/Source/Editor/UnrealEd/Private/DerivedDataBuildLocalExecutor.cpp b/Engine/Source/Editor/UnrealEd/Private/DerivedDataBuildLocalExecutor.cpp
--- a/Engine/Source/Editor/UnrealEd/Private/DerivedDataBuildLocalExecutor.cpp
+++ b/Engine/Source/Editor/UnrealEd/Private/DerivedDataBuildLocalExecutor.cpp
@@ -17,15 +17,11 @@
 #include "DerivedDataBuildOutput.h"
 #include "DerivedDataBuildTypes.h"
 #include "DerivedDataBuildWorker.h"
-#include "DerivedDataRequestOwner.h"
 #include "DerivedDataRequestTypes.h"
 #include "DerivedDataValue.h"
 #include "Features/IModularFeatures.h"
 #include "HAL/FileManager.h"
-#include "HAL/PlatformFile.h"
-#include "HAL/PlatformFileManager.h"
 #include "HAL/PlatformProcess.h"
-#include "Hash/xxhash.h"
 #include "IO/IoHash.h"
 #include "Logging/LogCategory.h"
 #include "Logging/LogMacros.h"
@@ -43,8 +39,6 @@
 #include "Serialization/CompactBinary.h"
 #include "Serialization/CompactBinarySerialization.h"
 #include "Serialization/CompactBinaryWriter.h"
-#include "Templates/Function.h"
-#include "Templates/Tuple.h"
 #include "Templates/UniquePtr.h"
 #include "Templates/UnrealTemplate.h"
 #include "Trace/Detail/Channel.h"
@@ -295,29 +289,8 @@ public:
 
 	[[nodiscard]] bool ManifestWorker(const FBuildWorker& Worker, FStringBuilderBase& OutPath)
 	{
-		// Calculate the worker hash that uniquely identifies the files of the worker.
-		FXxHash128Builder WorkerHashBuilder;
-		TArray<FIoHash> WorkerFileHashes;
-		TArray<TTuple<FUtf8StringView, bool>> WorkerFileMeta;
-
-		Worker.IterateExecutables([&WorkerHashBuilder, &WorkerFileHashes, &WorkerFileMeta](FUtf8StringView Path, const FIoHash& RawHash, uint64 RawSize)
-		{
-			WorkerHashBuilder.Update(MakeMemoryView(&RawHash, 1));
-			WorkerFileHashes.Emplace(RawHash);
-			WorkerFileMeta.Emplace(Path, true);
-		});
-
-		Worker.IterateFiles([&WorkerHashBuilder, &WorkerFileHashes, &WorkerFileMeta](FUtf8StringView Path, const FIoHash& RawHash, uint64 RawSize)
-		{
-			WorkerHashBuilder.Update(MakeMemoryView(&RawHash, 1));
-			WorkerFileHashes.Emplace(RawHash);
-			WorkerFileMeta.Emplace(Path, false);
-		});
-
-		const FXxHash128 WorkerHash = WorkerHashBuilder.Finalize();
-
 		OutPath.Reset();
-		FPathViews::Append(OutPath, SandboxRoot, TEXTVIEW("Workers"), WorkerHash);
+		FPathViews::Append(OutPath, SandboxRoot, TEXTVIEW("Workers"), Worker.GetKey().Hash);
 
 		// Return if a worker exists with this worker hash.
 		IFileManager& FileManager = IFileManager::Get();
@@ -332,7 +305,7 @@ public:
 		FManifestWorkerState* State;
 		{
 			TUniqueLock Lock(ManifestingWorkersMutex);
-			TUniquePtr<FManifestWorkerState>& StatePtr = ManifestingWorkers.FindOrAdd(WorkerHash);
+			TUniquePtr<FManifestWorkerState>& StatePtr = ManifestingWorkers.FindOrAdd(Worker.GetKey());
 			if (StatePtr)
 			{
 				bWaitForManifest = true;
@@ -349,7 +322,7 @@ public:
 			TUniqueLock Lock(ManifestingWorkersMutex);
 			if (--State->RefCount == 0)
 			{
-				ManifestingWorkers.Remove(WorkerHash);
+				ManifestingWorkers.Remove(Worker.GetKey());
 			}
 		};
 		if (bWaitForManifest)
@@ -360,67 +333,9 @@ public:
 
 		UE_LOGFMT(LogDerivedDataBuildLocalExecutor, Display,
 			"Manifesting build worker {WorkerName} in '{WorkerPath}'...", Worker.GetName(), OutPath);
-
-		// Manifest the worker in a temporary directory.
-		TStringBuilder<256> ScratchPath(InPlace, OutPath, '.', FGuid::NewGuid());
-
-		std::atomic<bool> bOk = true;
-		FRequestOwner BlockingOwner(EPriority::Blocking);
-		Worker.Resolve(WorkerFileHashes, BlockingOwner, [&Worker, &FileManager, &ScratchPath, &WorkerFileMeta, &bOk](FBuildWorkerResolvedParams&& Params)
-		{
-			uint32 MetaIndex = 0;
-			for (const FCompressedBuffer& Buffer : Params.Files)
-			{
-				TStringBuilder<256> FilePath;
-				FPathViews::Append(FilePath, ScratchPath, WorkerFileMeta[MetaIndex].Key);
-				if (TUniquePtr<FArchive> Ar{FileManager.CreateFileWriter(*FilePath, FILEWRITE_Silent)})
-				{
-					for (auto& Segment : Buffer.DecompressToComposite().GetSegments())
-					{
-						Ar->Serialize((void*)Segment.GetData(), Segment.GetSize());
-					}
-				}
-				else
-				{
-					UE_LOGFMT(LogDerivedDataBuildLocalExecutor, Display,
-						"Failed to open '{Path}' for writing for build worker {WorkerName}.",
-						FilePath, Worker.GetName());
-					bOk = false;
-				}
-
-				++MetaIndex;
-			}
-		});
-		BlockingOwner.Wait();
-
-		// Ensure this directory exists because the build worker cannot run correctly without it.
-		{
-			TStringBuilder<256> BinariesPath;
-			FPathViews::Append(BinariesPath, ScratchPath, TEXTVIEW("Engine/Binaries/Win64"));
-			if (!FileManager.MakeDirectory(*BinariesPath, /*bTree*/ true))
-			{
-				UE_LOGFMT(LogDerivedDataBuildLocalExecutor, Display,
-					"Failed to create directory '{Path}' for build worker {WorkerName}.",
-					BinariesPath, Worker.GetName());
-				bOk = false;
-			}
-		}
-
-		// Move the temporary directory to its final location or delete it if there was a race.
-		IPlatformFile& PlatformFile = FPlatformFileManager::Get().GetPlatformFile();
-		if (bOk && !PlatformFile.DirectoryExists(*OutPath))
-		{
-			PlatformFile.MoveFile(*OutPath, *ScratchPath);
-		}
-		PlatformFile.DeleteDirectoryRecursively(*ScratchPath);
-		const bool bExists = PlatformFile.DirectoryExists(*OutPath);
-		UE_CLOGFMT(!bExists, LogDerivedDataBuildLocalExecutor, Display,
-			"Failed to rename '{TempPath}' to '{Path}' for build worker {WorkerName}.",
-			ScratchPath, OutPath, Worker.GetName());
-
-		State->bOk = bOk;
+		State->bOk = Worker.Extract(OutPath);
 		State->ReadyEvent.Notify();
-		return bExists;
+		return State->bOk;
 	}
 
 	TConstArrayView<FUtf8StringView> GetHostPlatforms() const final
@@ -442,7 +357,7 @@ private:
 	};
 
 	FString SandboxRoot = FPaths::EngineSavedDir() / TEXT("LocalExec");
-	TMap<FXxHash128, TUniquePtr<FManifestWorkerState>> ManifestingWorkers;
+	TMap<FBuildWorkerKey, TUniquePtr<FManifestWorkerState>> ManifestingWorkers;
 	FMutex ManifestingWorkersMutex;
 };
 
