Write OidcToken output to user temp dir rather than project intermediate dir, to avoid issues when running under app translocation on Mac.

[RN] none

#Synced-CL 50570368

diff --git a/Engine/Source/Developer/DesktopPlatform/Private/DesktopPlatformBase.cpp b/Engine/Source/Developer/DesktopPlatform/Private/DesktopPlatformBase.cpp
--- a/Engine/Source/Developer/DesktopPlatform/Private/DesktopPlatformBase.cpp
+++ b/Engine/Source/Developer/DesktopPlatform/Private/DesktopPlatformBase.cpp
@@ -1032,8 +1032,7 @@ bool FDesktopPlatformBase::GetHordeAccessToken(const FString& HordeUrl, bool bUn
 bool FDesktopPlatformBase::GetOidcAccessTokenInternal(const FString& RootDir, const FString& BaseArguments, EDesktopLoginInteractionLevel InteractionLevel, FFeedbackContext* Warn, FString& OutToken, FDateTime& OutTokenExpiresAt, bool& bOutWasInteractiveLogin)
 {
 	#if !defined(UE_NO_ENGINE_OIDC) || UE_NO_ENGINE_OIDC == 0
-	IFileManager::Get().MakeDirectory(*FPaths::ProjectIntermediateDir(), /*bTree*/ true);
-	FString ResultFilePath = FPaths::CreateTempFilename(*FPaths::ProjectIntermediateDir(), TEXT("oidcToken.json"));
+	FString ResultFilePath = FPaths::CreateTempFilename(FPlatformProcess::UserTempDir(), TEXT("oidcToken.json"));
 	#else
 	FString ResultFilePath = TEXT("oidcToken.json");
 	#endif
