Call BuildUBT.bat rather than msbuild

#rnx

#Synced-CL 51424418

diff --git a/Engine/Source/Developer/DesktopPlatform/Private/DesktopPlatformBase.cpp b/Engine/Source/Developer/DesktopPlatform/Private/DesktopPlatformBase.cpp
--- a/Engine/Source/Developer/DesktopPlatform/Private/DesktopPlatformBase.cpp
+++ b/Engine/Source/Developer/DesktopPlatform/Private/DesktopPlatformBase.cpp
@@ -1877,91 +1877,6 @@ bool FDesktopPlatformBase::EnumerateProjectsKnownByEngine(const FString &Identif
 	return true;
 }
 
-#if PLATFORM_WINDOWS
-
-#include "Windows/AllowWindowsPlatformTypes.h"
-#include <ShlObj.h>
-#include "Windows/HideWindowsPlatformTypes.h"
-
-static bool TryReadMsBuildInstallPath(HKEY RootKey, const TCHAR* KeyName, const TCHAR* ValueName, const TCHAR* MsBuildRelativePath, FString& OutMsBuildPath)
-{
-	FString Value;
-	if (!FWindowsPlatformMisc::QueryRegKey(RootKey, KeyName, ValueName, Value))
-	{
-		return false;
-	}
-
-	FString Result = Value / MsBuildRelativePath;
-	if (!FPaths::FileExists(Result))
-	{
-		return false;
-	}
-
-	OutMsBuildPath = Result;
-	return true;
-}
-
-static bool TryReadMsBuildInstallPath(const TCHAR* KeyRelativeName, const TCHAR* ValueName, const TCHAR* MsBuildRelativePath, FString& OutMsBuildPath)
-{
-	if (TryReadMsBuildInstallPath(HKEY_CURRENT_USER, *(FString("SOFTWARE\\") + KeyRelativeName), ValueName, MsBuildRelativePath, OutMsBuildPath))
-	{
-		return true;
-	}
-	if (TryReadMsBuildInstallPath(HKEY_LOCAL_MACHINE, *(FString("SOFTWARE\\") + KeyRelativeName), ValueName, MsBuildRelativePath, OutMsBuildPath))
-	{
-		return true;
-	}
-	if (TryReadMsBuildInstallPath(HKEY_CURRENT_USER, *(FString("SOFTWARE\\Wow6432Node\\") + KeyRelativeName), ValueName, MsBuildRelativePath, OutMsBuildPath))
-	{
-		return true;
-	}
-	if (TryReadMsBuildInstallPath(HKEY_LOCAL_MACHINE, *(FString("SOFTWARE\\Wow6432Node\\") + KeyRelativeName), ValueName, MsBuildRelativePath, OutMsBuildPath))
-	{
-		return true;
-	}
-	return false;
-}
-
-static bool TryReadMsBuildInstallPath(FString& OutPath)
-{
-	// Try to get the MSBuild 14.0 path directly (see https://msdn.microsoft.com/en-us/library/hh162058(v=vs.120).aspx)
-	TCHAR ProgramFilesX86[MAX_PATH];
-	if (SUCCEEDED(SHGetFolderPath(NULL, CSIDL_PROGRAM_FILES | CSIDL_FLAG_CREATE, NULL, SHGFP_TYPE_CURRENT, ProgramFilesX86)))
-	{
-		FString ToolPath = FString(ProgramFilesX86) / TEXT("MSBuild/14.0/bin/MSBuild.exe");
-		if (FPaths::FileExists(ToolPath))
-		{
-			OutPath = ToolPath;
-			return true;
-		}
-	}
-
-	// Try to get the MSBuild 14.0 path from the registry
-	if (TryReadMsBuildInstallPath(TEXT("Microsoft\\MSBuild\\ToolsVersions\\14.0"), TEXT("MSBuildToolsPath"), TEXT("MSBuild.exe"), OutPath))
-	{
-		return true;
-	}
-
-	// Check for MSBuild 15. This is installed alongside Visual Studio 2017, so we get the path relative to that.
-	if (TryReadMsBuildInstallPath(TEXT("Microsoft\\VisualStudio\\SxS\\VS7"), TEXT("15.0"), TEXT("MSBuild\\15.0\\bin\\MSBuild.exe"), OutPath))
-	{
-		return true;
-	}
-
-	// Check for older versions of MSBuild. These are registered as separate versions in the registry.
-	if (TryReadMsBuildInstallPath(TEXT("Microsoft\\MSBuild\\ToolsVersions\\12.0"), TEXT("MSBuildToolsPath"), TEXT("MSBuild.exe"), OutPath))
-	{
-		return true;
-	}
-	if (TryReadMsBuildInstallPath(TEXT("Microsoft\\MSBuild\\ToolsVersions\\4.0"), TEXT("MSBuildToolsPath"), TEXT("MSBuild.exe"), OutPath))
-	{
-		return true;
-	}
-
-	return false;
-}
-#endif
-
 bool FDesktopPlatformBase::BuildUnrealBuildTool(const FString& RootDir, FOutputDevice& Ar)
 {
 	Ar.Logf(TEXT("Building UnrealBuildTool in %s..."), *RootDir);
@@ -1977,28 +1892,23 @@ bool FDesktopPlatformBase::BuildUnrealBuildTool(const FString& RootDir, FOutputD
 	FString CompilerExecutableFilename;
 	FString CmdLineParams;
 
-	if (PLATFORM_WINDOWS)
+	if constexpr (PLATFORM_WINDOWS)
 	{
-#if PLATFORM_WINDOWS
-		if (!TryReadMsBuildInstallPath(CompilerExecutableFilename))
-		{
-			Ar.Logf(TEXT("Couldn't find MSBuild installation; skipping."));
-			return false;
-		}
-#endif
-		CmdLineParams = FString::Printf(TEXT("/nologo /verbosity:quiet \"%s\" /property:Configuration=Development /property:Platform=AnyCPU"), *CsProjLocation);
+		FString ScriptPath = FPaths::ConvertRelativePathToFull(RootDir / TEXT("Engine/Build/BatchFiles/BuildUBT.bat"));
+		CompilerExecutableFilename = TEXT("cmd.exe");
+		CmdLineParams = FString::Printf(TEXT("/c \"%s\""), *ScriptPath);
 	}
-	else if (PLATFORM_MAC)
+	else if constexpr (PLATFORM_MAC)
 	{
-		FString ScriptPath = FPaths::ConvertRelativePathToFull(RootDir / TEXT("Engine/Build/BatchFiles/RunDotnetMSBuild.sh"));
+		FString ScriptPath = FPaths::ConvertRelativePathToFull(RootDir / TEXT("Engine/Build/BatchFiles/BuildUBT.sh"));
 		CompilerExecutableFilename = TEXT("/bin/sh");
-		CmdLineParams = FString::Printf(TEXT("\"%s\" /property:Configuration=Development %s"), *ScriptPath, *CsProjLocation);
+		CmdLineParams = FString::Printf(TEXT("\"%s\""), *ScriptPath);
 	}
-	else if (PLATFORM_LINUX)
+	else if constexpr (PLATFORM_LINUX)
 	{
-		FString ScriptPath = FPaths::ConvertRelativePathToFull(RootDir / TEXT("Engine/Build/BatchFiles/RunDotnetMSBuild.sh"));
+		FString ScriptPath = FPaths::ConvertRelativePathToFull(RootDir / TEXT("Engine/Build/BatchFiles/BuildUBT.sh"));
 		CompilerExecutableFilename = TEXT("/bin/bash");
-		CmdLineParams = FString::Printf(TEXT("\"%s\" /property:Configuration=Development %s"), *ScriptPath, *CsProjLocation);
+		CmdLineParams = FString::Printf(TEXT("\"%s\""), *ScriptPath);
 	}
 	else
 	{
