[DesktopPlatform]

Don't consider target files that aren't in the cache for determining if we need to rebuild via UBT. The target gathering on disk can fetch Target.cs files for programs that have no impact on the running program nor will they be appended to any newly generated TargetInfo.json causing the fast, cache hit, path from ever being taken.

Now we only consider modifications for the target files we expect. Any new files found are assumed to be irrelevant since a new target.cs would require a proper build first and would thus be in the TargetInfo.json te next time we run anyway.

#rnx

#Synced-CL 51343429

diff --git a/Engine/Source/Developer/DesktopPlatform/Private/DesktopPlatformBase.cpp b/Engine/Source/Developer/DesktopPlatform/Private/DesktopPlatformBase.cpp
--- a/Engine/Source/Developer/DesktopPlatform/Private/DesktopPlatformBase.cpp
+++ b/Engine/Source/Developer/DesktopPlatform/Private/DesktopPlatformBase.cpp
@@ -1374,10 +1374,12 @@ struct FTargetFileVisitor : IPlatformFile::FDirectoryStatVisitor
 				UE_LOG(LogDesktopPlatform, Log, TEXT("Found target file %s newer than cache"), InFileNameOrDirectory);
 				bCacheInvalid = true;
 			}
+			// Any target file found on disk not in the TargetInfo.json is assumed to not be 
+			// relevant for the current build since anything new on disk hasn't been built at all. 
 			if (RemainingTargetNames.Remove(TargetName) != 1)
 			{
-				UE_LOG(LogDesktopPlatform, Log, TEXT("Found target file %s not present in cache"), InFileNameOrDirectory);
-				bCacheInvalid = true;
+				UE_LOG(LogDesktopPlatform, Log, TEXT("Found target file %s not present in cache. Ignoring."), InFileNameOrDirectory);
+				return true;
 			}
 			
 			return !bCacheInvalid;
