ignore leading whitespace for mac file dialog file extensions

#Synced-CL 52243953

diff --git a/Engine/Source/Developer/DesktopPlatform/Private/Mac/DesktopPlatformMac.cpp b/Engine/Source/Developer/DesktopPlatform/Private/Mac/DesktopPlatformMac.cpp
--- a/Engine/Source/Developer/DesktopPlatform/Private/Mac/DesktopPlatformMac.cpp
+++ b/Engine/Source/Developer/DesktopPlatform/Private/Mac/DesktopPlatformMac.cpp
@@ -159,8 +159,21 @@ private:
 
 		for( int32 Index = 0; Index < [ExtensionsWildcards count]; ++Index )
 		{
-			NSString* Temp = [[ExtensionsWildcards objectAtIndex:Index] stringByTrimmingCharactersInSet:[NSCharacterSet characterSetWithCharactersInString:@"*."]];
-			[Extensions addObject: [UTType typeWithFilenameExtension:Temp]];
+			// capture potential leading whitespace and *.
+			NSMutableCharacterSet* TrimSet = [NSMutableCharacterSet characterSetWithCharactersInString:@"*."];                                                    
+			[TrimSet formUnionWithCharacterSet:[NSCharacterSet whitespaceCharacterSet]];                                                                          
+			NSString* Temp = [[ExtensionsWildcards objectAtIndex:Index] stringByTrimmingCharactersInSet:TrimSet];
+			
+			UTType* Type = [UTType typeWithFilenameExtension:Temp];
+			
+			if(Type == nil)
+			{
+				UE_LOGF(LogMac, Warning, "Couldn't find UTType for requested file picker extension extension %hs", [Temp UTF8String]);
+			}
+			else
+			{
+				[Extensions addObject: Type];
+			}
 		}
 
 		DialogPanel.allowedContentTypes = Extensions;
@@ -427,14 +440,17 @@ bool FDesktopPlatformMac::FileDialogShared(bool bSave, const void* ParentWindowH
 					CFRelease(Type);
 				}
 			}
-
-			if ([AllowedFileTypes count] == 1)
+			
+			if(AllowedFileTypes != nil)
 			{
-				// Duplicate the type for the description
-				[AllowedFileTypes addObject:AllowedFileTypes[0]];
+				if ([AllowedFileTypes count] == 1)
+				{
+					// Duplicate the type for the description
+					[AllowedFileTypes addObject:AllowedFileTypes[0]];
+				}
+				
+				[AccessoryView AddAllowedFileTypes:AllowedFileTypes];
 			}
-			
-			[AccessoryView AddAllowedFileTypes:AllowedFileTypes];
 
 			bool bOkPressed = false;
 			NSWindow* FocusWindow = [[NSApplication sharedApplication] keyWindow];
